Fix some non-unity/non-pch errors from Dev-Rendering and Dev-VR merges.

diff --git a/Engine/Plugins/Online/OnlineSubsystemOculus/Source/Private/OculusFindSessionsCallbackProxy.cpp b/Engine/Plugins/Online/OnlineSubsystemOculus/Source/Private/OculusFindSessionsCallbackProxy.cpp
--- a/Engine/Plugins/Online/OnlineSubsystemOculus/Source/Private/OculusFindSessionsCallbackProxy.cpp
+++ b/Engine/Plugins/Online/OnlineSubsystemOculus/Source/Private/OculusFindSessionsCallbackProxy.cpp
@@ -4,6 +4,7 @@
 #include "CoreOnline.h"
 #include "Online.h"
 #include "OnlineSessionInterfaceOculus.h"
+#include "OnlineSubsystemOculusPrivate.h"
 
 UOculusFindSessionsCallbackProxy::UOculusFindSessionsCallbackProxy(const FObjectInitializer& ObjectInitializer)
 	: Super(ObjectInitializer)
diff --git a/Engine/Plugins/Online/OnlineSubsystemOculus/Source/Private/OnlineVoiceOculus.cpp b/Engine/Plugins/Online/OnlineSubsystemOculus/Source/Private/OnlineVoiceOculus.cpp
--- a/Engine/Plugins/Online/OnlineSubsystemOculus/Source/Private/OnlineVoiceOculus.cpp
+++ b/Engine/Plugins/Online/OnlineSubsystemOculus/Source/Private/OnlineVoiceOculus.cpp
@@ -4,6 +4,7 @@
 #include "Voice.h"
 #include "OnlineSubsystemOculus.h"
 #include "Components/AudioComponent.h"
+#include "Misc/ConfigCacheIni.h"
 
 #include "Runtime/Engine/Classes/Sound/SoundWaveProcedural.h"
 #include "OnlineSubsystemUtils.h"
diff --git a/Engine/Source/Editor/DetailCustomizations/Private/MacTargetSettingsDetails.h b/Engine/Source/Editor/DetailCustomizations/Private/MacTargetSettingsDetails.h
--- a/Engine/Source/Editor/DetailCustomizations/Private/MacTargetSettingsDetails.h
+++ b/Engine/Source/Editor/DetailCustomizations/Private/MacTargetSettingsDetails.h
@@ -8,6 +8,7 @@
 class IDetailLayoutBuilder;
 class IPropertyHandle;
 class SErrorText;
+class SWidget;
 
 enum class ECheckBoxState : uint8;
 
diff --git a/Engine/Source/Runtime/Engine/Private/PrimitiveSceneProxy.cpp b/Engine/Source/Runtime/Engine/Private/PrimitiveSceneProxy.cpp
--- a/Engine/Source/Runtime/Engine/Private/PrimitiveSceneProxy.cpp
+++ b/Engine/Source/Runtime/Engine/Private/PrimitiveSceneProxy.cpp
@@ -11,6 +11,7 @@
 #include "Components/BrushComponent.h"
 #include "SceneManagement.h"
 #include "PrimitiveSceneInfo.h"
+#include "Materials/Material.h"
 
 static TAutoConsoleVariable<int32> CVarForceSingleSampleShadowingFromStationary(
 	TEXT("r.Shadow.ForceSingleSampleShadowingFromStationary"),
diff --git a/Engine/Source/Runtime/Slate/Private/Framework/Application/MenuStack.cpp b/Engine/Source/Runtime/Slate/Private/Framework/Application/MenuStack.cpp
--- a/Engine/Source/Runtime/Slate/Private/Framework/Application/MenuStack.cpp
+++ b/Engine/Source/Runtime/Slate/Private/Framework/Application/MenuStack.cpp
@@ -7,6 +7,7 @@
 #include "Widgets/Text/STextBlock.h"
 #include "Widgets/Layout/SBox.h"
 #include "Widgets/Layout/SPopup.h"
+#include "Framework/Application/Menu.h"
 
 #define LOCTEXT_NAMESPACE "MenuStack"
 
