More build fixes.

diff --git a/Engine/Plugins/Online/OnlineSubsystemOculus/Source/Private/OnlineSubsystemOculusTypes.h b/Engine/Plugins/Online/OnlineSubsystemOculus/Source/Private/OnlineSubsystemOculusTypes.h
--- a/Engine/Plugins/Online/OnlineSubsystemOculus/Source/Private/OnlineSubsystemOculusTypes.h
+++ b/Engine/Plugins/Online/OnlineSubsystemOculus/Source/Private/OnlineSubsystemOculusTypes.h
@@ -96,7 +96,7 @@ public:
 	/** Needed for TMap::GetTypeHash() */
 	friend uint32 GetTypeHash(const FUniqueNetIdOculus& A)
 	{
-		return GetTypeHash(A.ID);
+		return GetTypeHash((uint64)A.ID);
 	}
 };
 
diff --git a/Engine/Source/Programs/UnrealBuildTool/Android/AndroidToolChain.cs b/Engine/Source/Programs/UnrealBuildTool/Android/AndroidToolChain.cs
--- a/Engine/Source/Programs/UnrealBuildTool/Android/AndroidToolChain.cs
+++ b/Engine/Source/Programs/UnrealBuildTool/Android/AndroidToolChain.cs
@@ -820,6 +820,8 @@ namespace UnrealBuildTool
 			}
 
 			var Result = new List<string>();
+			Result.Add("#include \"CoreTypes.h\"");
+			Result.Add("");
 			foreach (string Arch in Arches)
 			{
 				switch (Arch)
