Fixing shadowing.

diff --git a/Engine/Plugins/Online/OnlineSubsystemOculus/Source/Private/OnlineSessionInterfaceOculus.cpp b/Engine/Plugins/Online/OnlineSubsystemOculus/Source/Private/OnlineSessionInterfaceOculus.cpp
--- a/Engine/Plugins/Online/OnlineSubsystemOculus/Source/Private/OnlineSessionInterfaceOculus.cpp
+++ b/Engine/Plugins/Online/OnlineSubsystemOculus/Source/Private/OnlineSessionInterfaceOculus.cpp
@@ -348,8 +348,8 @@ bool FOnlineSessionOculus::UpdateMatchmakingRoom(FName SessionName, FOnlineSessi
 				return;
 			}
 
-			auto Session = GetNamedSession(SessionName);
-			if (Session == nullptr)
+			auto NewSession = GetNamedSession(SessionName);
+			if (NewSession == nullptr)
 			{
 				UE_LOG_ONLINE(Error, TEXT("Session (%s) no longer exists"), *SessionName.ToString());
 				TriggerOnUpdateSessionCompleteDelegates(SessionName, false);
@@ -357,7 +357,7 @@ bool FOnlineSessionOculus::UpdateMatchmakingRoom(FName SessionName, FOnlineSessi
 			}
 
 			// Update the Session Settings
-			Session->SessionSettings.bShouldAdvertise = !Session->SessionSettings.bShouldAdvertise;
+			NewSession->SessionSettings.bShouldAdvertise = !NewSession->SessionSettings.bShouldAdvertise;
 
 			UpdateRoomDataStore(SessionName, UpdatedSessionSettings);
 		}));
@@ -446,8 +446,8 @@ bool FOnlineSessionOculus::UpdateRoomDataStore(FName SessionName, FOnlineSession
 				return;
 			}
 
-			auto Session = GetNamedSession(SessionName);
-			if (Session == nullptr)
+			auto NewSession = GetNamedSession(SessionName);
+			if (NewSession == nullptr)
 			{
 				UE_LOG_ONLINE(Error, TEXT("Session (%s) no longer exists"), *SessionName.ToString());
 				TriggerOnUpdateSessionCompleteDelegates(SessionName, false);
@@ -456,7 +456,7 @@ bool FOnlineSessionOculus::UpdateRoomDataStore(FName SessionName, FOnlineSession
 
 			// Update the Room
 			auto Room = ovr_Message_GetRoom(Message);
-			UpdateSessionFromRoom(*Session, Room);
+			UpdateSessionFromRoom(*NewSession, Room);
 
 			TriggerOnUpdateSessionCompleteDelegates(SessionName, true);
 		}));
