Fix PVS-Studio warnings from Dev-VR.

diff --git a/Engine/Plugins/Online/OnlineSubsystemOculus/Source/Private/OnlineVoiceOculus.h b/Engine/Plugins/Online/OnlineSubsystemOculus/Source/Private/OnlineVoiceOculus.h
--- a/Engine/Plugins/Online/OnlineSubsystemOculus/Source/Private/OnlineVoiceOculus.h
+++ b/Engine/Plugins/Online/OnlineSubsystemOculus/Source/Private/OnlineVoiceOculus.h
@@ -70,7 +70,7 @@ public:
 	}
 	virtual bool IsLocalPlayerTalking(uint32 LocalUserNum) override
 	{
-		return (!bIsLocalPlayerMuted && 0 <= LocalUserNum && LocalUserNum <= MAX_LOCAL_PLAYERS);
+		return (!bIsLocalPlayerMuted && LocalUserNum <= MAX_LOCAL_PLAYERS);
 	}
 	virtual bool IsRemotePlayerTalking(const FUniqueNetId& UniqueId) override;
 	virtual bool IsMuted(uint32 LocalUserNum, const FUniqueNetId& UniqueId) const override;
diff --git a/Engine/Plugins/Runtime/Oculus/OculusVR/Source/OculusHMD/Private/OculusHMDPrivate.cpp b/Engine/Plugins/Runtime/Oculus/OculusVR/Source/OculusHMD/Private/OculusHMDPrivate.cpp
--- a/Engine/Plugins/Runtime/Oculus/OculusVR/Source/OculusHMD/Private/OculusHMDPrivate.cpp
+++ b/Engine/Plugins/Runtime/Oculus/OculusVR/Source/OculusHMD/Private/OculusHMDPrivate.cpp
@@ -75,9 +75,9 @@ void ExecuteOnRenderThread(const std::function<void()>& Function)
 	{
 		ENQUEUE_UNIQUE_RENDER_COMMAND_ONEPARAMETER(
 			ExecuteOnRenderThread,
-			std::function<void()>, function, Function,
+			std::function<void()>, FunctionArg, Function,
 			{
-				function();
+				FunctionArg();
 			});
 
 		FlushRenderingCommands();
