Fixes for remaining issues flagged by PVS Studio:

* Null pointer being passed to ExportPolys without being checked.
* Incorrect test against UVREditorMotionControllerInteractor spline index which is out of bounds.
* FOnlineVoiceSteam::UnmuteRemoteTalker() never setting return value for success cases.
* Object not freed on return from SetApexDestructibleAsset()
* FInternetAddrOculus instances being leaked from UOculusNetDriver::InitConnect()
* FOnlineMessageMultiTaskOculusWriteAchievements instanes being leaked.

Following warnings temporarily disabled pending review by original authors:

* Potential OOB array access in StaticMeshEditor.cpp (NickP)
* bHasOverlap check is never true in BitmapUtils.h (MichaelN)
* Potential null pointer dereference in NetworkPlatformFile.cpp (DanielL)

diff --git a/Engine/Plugins/2D/Paper2D/Source/Paper2D/Private/BitmapUtils.h b/Engine/Plugins/2D/Paper2D/Source/Paper2D/Private/BitmapUtils.h
--- a/Engine/Plugins/2D/Paper2D/Source/Paper2D/Private/BitmapUtils.h
+++ b/Engine/Plugins/2D/Paper2D/Source/Paper2D/Private/BitmapUtils.h
@@ -199,7 +199,7 @@ struct FBitmap : public FAlphaBitmap
 									break;
 								}
 							}
-						} while (bHasOverlap);
+						} while (bHasOverlap); //-V654
 
 						// Mark in the mask, avoid checking for any more sprites within this rect
 						MaskBitmap.FillRect(NewRect.Min.X, NewRect.Min.Y, NewRect.Width(), NewRect.Height());
diff --git a/Engine/Plugins/Online/OnlineSubsystemOculus/Source/Private/OculusNetDriver.cpp b/Engine/Plugins/Online/OnlineSubsystemOculus/Source/Private/OculusNetDriver.cpp
--- a/Engine/Plugins/Online/OnlineSubsystemOculus/Source/Private/OculusNetDriver.cpp
+++ b/Engine/Plugins/Online/OnlineSubsystemOculus/Source/Private/OculusNetDriver.cpp
@@ -96,8 +96,8 @@ bool UOculusNetDriver::InitConnect(FNetworkNotify* InNotify, const FURL& Connect
 {
 	UE_LOG(LogNet, Verbose, TEXT("Connecting to host: %s"), *ConnectURL.ToString(true));
 
-	auto OculusAddr = new FInternetAddrOculus(ConnectURL);
-	if (!OculusAddr->IsValid())
+	FInternetAddrOculus OculusAddr(ConnectURL);
+	if (!OculusAddr.IsValid())
 	{
 		UE_LOG(LogNet, Verbose, TEXT("Init as IPNetDriver connect"));
 		bIsPassthrough = true;
@@ -115,12 +115,12 @@ bool UOculusNetDriver::InitConnect(FNetworkNotify* InNotify, const FURL& Connect
 
 	// Set it as the server connection before anything else so everything knows this is a client
 	ServerConnection = Connection;
-	Connection->InitLocalConnection(this, nullptr, ConnectURL, ovr_Net_IsConnected(OculusAddr->GetID()) ? USOCK_Open : USOCK_Pending);
+	Connection->InitLocalConnection(this, nullptr, ConnectURL, ovr_Net_IsConnected(OculusAddr.GetID()) ? USOCK_Open : USOCK_Pending);
 
-	Connections.Add(OculusAddr->GetID(), Connection);
+	Connections.Add(OculusAddr.GetID(), Connection);
 
 	// Connect via OVR_Networking
-	ovr_Net_Connect(OculusAddr->GetID());
+	ovr_Net_Connect(OculusAddr.GetID());
 
 	// Create the control channel so we can send the Hello message
 	ServerConnection->CreateChannel(CHTYPE_Control, true, INDEX_NONE);
diff --git a/Engine/Plugins/Online/OnlineSubsystemOculus/Source/Private/OnlineAchievementsInterfaceOculus.cpp b/Engine/Plugins/Online/OnlineSubsystemOculus/Source/Private/OnlineAchievementsInterfaceOculus.cpp
--- a/Engine/Plugins/Online/OnlineSubsystemOculus/Source/Private/OnlineAchievementsInterfaceOculus.cpp
+++ b/Engine/Plugins/Online/OnlineSubsystemOculus/Source/Private/OnlineAchievementsInterfaceOculus.cpp
@@ -5,8 +5,9 @@
 #include "OnlineIdentityOculus.h"
 #include "OnlineMessageMultiTaskOculus.h"
 #include "OnlineSubsystemOculusPackage.h"
+#include "Templates/SharedPointer.h"
 
-class FOnlineMessageMultiTaskOculusWriteAchievements : public FOnlineMessageMultiTaskOculus
+class FOnlineMessageMultiTaskOculusWriteAchievements : public FOnlineMessageMultiTaskOculus, public TSharedFromThis<FOnlineMessageMultiTaskOculusWriteAchievements>
 {
 private:
 
@@ -14,32 +15,56 @@ private:
 	FOnlineAchievementsWriteRef WriteObject;
 	FOnAchievementsWrittenDelegate AchievementDelegate;
 
+	// private to force the use of FOnlineMessageMultiTaskOculusWriteAchievements::Create()
+	FOnlineMessageMultiTaskOculusWriteAchievements(FOnlineSubsystemOculus& InOculusSubsystem, const FUniqueNetIdOculus& InPlayerId, FOnlineAchievementsWriteRef& InWriteObject, const FOnAchievementsWrittenDelegate& InAchievementDelegate)
+		: FOnlineMessageMultiTaskOculus(InOculusSubsystem, FOnlineMessageMultiTaskOculus::FFinalizeDelegate::CreateRaw(this, &FOnlineMessageMultiTaskOculusWriteAchievements::Finalize))
+		, PlayerId(InPlayerId)
+		, WriteObject(InWriteObject)
+		, AchievementDelegate(InAchievementDelegate)
+	{}
+
+	static TSet< TSharedRef<FOnlineMessageMultiTaskOculusWriteAchievements> > ActiveAchievementWriteTasks;
+
 PACKAGE_SCOPE:
 
-	FOnlineMessageMultiTaskOculusWriteAchievements(
+	static TSharedRef<FOnlineMessageMultiTaskOculusWriteAchievements> Create(
 		FOnlineSubsystemOculus& InOculusSubsystem,
 		const FUniqueNetIdOculus& InPlayerId,
 		FOnlineAchievementsWriteRef& InWriteObject,
-		const FOnAchievementsWrittenDelegate& InAchievementDelegate) :
-		FOnlineMessageMultiTaskOculus(InOculusSubsystem, FOnlineMessageMultiTaskOculus::FFinalizeDelegate::CreateRaw(this, &FOnlineMessageMultiTaskOculusWriteAchievements::Finalize)),
-		PlayerId(InPlayerId),
-		WriteObject(InWriteObject),
-		AchievementDelegate(InAchievementDelegate)
+		const FOnAchievementsWrittenDelegate& InAchievementDelegate)
 	{
+		TSharedRef<FOnlineMessageMultiTaskOculusWriteAchievements> NewTask = MakeShareable(new FOnlineMessageMultiTaskOculusWriteAchievements(InOculusSubsystem, InPlayerId, InWriteObject, InAchievementDelegate));
+		ActiveAchievementWriteTasks.Add(NewTask);
+
+		return NewTask;
 	}
 
 	void Finalize()
 	{
 		WriteObject->WriteState = (bDidAllRequestsFinishedSuccessfully) ? EOnlineAsyncTaskState::Done : EOnlineAsyncTaskState::Failed;
 		AchievementDelegate.ExecuteIfBound(PlayerId, true);
+
+		// this should delete this task object, make sure it happens last
+		ActiveAchievementWriteTasks.Remove(AsShared());
+	}
+
+	static void ClearAllActiveTasks()
+	{
+		ActiveAchievementWriteTasks.Empty();
 	}
 };
+TSet< TSharedRef<FOnlineMessageMultiTaskOculusWriteAchievements> > FOnlineMessageMultiTaskOculusWriteAchievements::ActiveAchievementWriteTasks;
 
 FOnlineAchievementsOculus::FOnlineAchievementsOculus(class FOnlineSubsystemOculus& InSubsystem)
 : OculusSubsystem(InSubsystem)
 {
 }
 
+FOnlineAchievementsOculus::~FOnlineAchievementsOculus()
+{
+	FOnlineMessageMultiTaskOculusWriteAchievements::ClearAllActiveTasks();
+}
+
 void FOnlineAchievementsOculus::WriteAchievements(const FUniqueNetId& PlayerId, FOnlineAchievementsWriteRef& WriteObject, const FOnAchievementsWrittenDelegate& Delegate)
 {
 	if (AchievementDescriptions.Num() == 0)
@@ -68,7 +93,7 @@ void FOnlineAchievementsOculus::WriteAchievements(const FUniqueNetId& PlayerId,
 	}
 
 	WriteObject->WriteState = EOnlineAsyncTaskState::InProgress;
-	auto MultiTask = new FOnlineMessageMultiTaskOculusWriteAchievements(OculusSubsystem, static_cast<FUniqueNetIdOculus>(PlayerId), WriteObject, Delegate);
+	TSharedRef<FOnlineMessageMultiTaskOculusWriteAchievements> MultiTask = FOnlineMessageMultiTaskOculusWriteAchievements::Create(OculusSubsystem, static_cast<FUniqueNetIdOculus>(PlayerId), WriteObject, Delegate);
 
 	// treat each achievement as unlocked
 	for (FStatPropertyArray::TConstIterator It(WriteObject->Properties); It; ++It)
diff --git a/Engine/Plugins/Online/OnlineSubsystemOculus/Source/Private/OnlineAchievementsInterfaceOculus.h b/Engine/Plugins/Online/OnlineSubsystemOculus/Source/Private/OnlineAchievementsInterfaceOculus.h
--- a/Engine/Plugins/Online/OnlineSubsystemOculus/Source/Private/OnlineAchievementsInterfaceOculus.h
+++ b/Engine/Plugins/Online/OnlineSubsystemOculus/Source/Private/OnlineAchievementsInterfaceOculus.h
@@ -105,7 +105,7 @@ public:
 	/**
 	* Default destructor
 	*/
-	virtual ~FOnlineAchievementsOculus() = default;
+	virtual ~FOnlineAchievementsOculus();
 
 	// Begin IOnlineAchievements interface
 	virtual void WriteAchievements(const FUniqueNetId& PlayerId, FOnlineAchievementsWriteRef& WriteObject, const FOnAchievementsWrittenDelegate& Delegate = FOnAchievementsWrittenDelegate()) override;
diff --git a/Engine/Plugins/Online/OnlineSubsystemSteam/Source/Private/VoiceInterfaceSteam.cpp b/Engine/Plugins/Online/OnlineSubsystemSteam/Source/Private/VoiceInterfaceSteam.cpp
--- a/Engine/Plugins/Online/OnlineSubsystemSteam/Source/Private/VoiceInterfaceSteam.cpp
+++ b/Engine/Plugins/Online/OnlineSubsystemSteam/Source/Private/VoiceInterfaceSteam.cpp
@@ -439,6 +439,7 @@ bool FOnlineVoiceSteam::MuteRemoteTalker(uint8 LocalUserNum, const FUniqueNetId&
 			SystemMuteList.AddUnique((const FUniqueNetIdSteam&)PlayerId);
 			// Should update MuteList after going up to the server and coming back down
 			ProcessMuteChangeNotification();
+			Return = S_OK;
 		}
 		else
 		{
@@ -481,6 +482,7 @@ bool FOnlineVoiceSteam::UnmuteRemoteTalker(uint8 LocalUserNum, const FUniqueNetI
 			SystemMuteList.RemoveSingleSwap((const FUniqueNetIdSteam&)PlayerId);
 			// Should update MuteList after going up to the server and coming back down
 			ProcessMuteChangeNotification();
+			Return = S_OK;
 		}
 		else
 		{
@@ -501,6 +503,7 @@ bool FOnlineVoiceSteam::UnmuteRemoteTalker(uint8 LocalUserNum, const FUniqueNetI
 						MuteList.RemoveSingleSwap((const FUniqueNetIdSteam&)PlayerId);
 						UE_LOG(LogVoice, Log, TEXT("Unmuting remote talker (%s)"), *PlayerId.ToDebugString());
 					}
+					Return = S_OK;
 				}
 				else
 				{
diff --git a/Engine/Plugins/Runtime/ApexDestruction/Source/ApexDestruction/Private/ApexDestructibleAssetImport.cpp b/Engine/Plugins/Runtime/ApexDestruction/Source/ApexDestruction/Private/ApexDestructibleAssetImport.cpp
--- a/Engine/Plugins/Runtime/ApexDestruction/Source/ApexDestruction/Private/ApexDestructibleAssetImport.cpp
+++ b/Engine/Plugins/Runtime/ApexDestruction/Source/ApexDestruction/Private/ApexDestructibleAssetImport.cpp
@@ -855,6 +855,12 @@ bool SetApexDestructibleAsset(UDestructibleMesh& DestructibleMesh, apex::Destruc
 	int32 SkeletalDepth=0;
 	if(!ProcessImportMeshSkeleton(DestructibleMesh.Skeleton, DestructibleMesh.RefSkeleton, SkeletalDepth, *SkelMeshImportDataPtr))
 	{
+		if (ExistDestMeshDataPtr)
+		{
+			RestoreExistingDestMeshData(ExistDestMeshDataPtr, &DestructibleMesh);
+			delete ExistDestMeshDataPtr;
+			ExistDestMeshDataPtr = NULL;
+		}
 		return false;
 	}
 	UE_LOG(LogApexDestructibleAssetImport, Warning, TEXT("Bones digested - %i  Depth of hierarchy - %i"), DestructibleMesh.RefSkeleton.GetNum(), SkeletalDepth);
diff --git a/Engine/Source/Editor/StaticMeshEditor/Private/StaticMeshEditor.cpp b/Engine/Source/Editor/StaticMeshEditor/Private/StaticMeshEditor.cpp
--- a/Engine/Source/Editor/StaticMeshEditor/Private/StaticMeshEditor.cpp
+++ b/Engine/Source/Editor/StaticMeshEditor/Private/StaticMeshEditor.cpp
@@ -1156,9 +1156,9 @@ TSharedRef<SWidget> FStaticMeshEditor::GenerateUVChannelComboList()
 
 void FStaticMeshEditor::UpdateLODStats(int32 CurrentLOD) 
 {
-	NumTriangles[CurrentLOD] = 0;
-	NumVertices[CurrentLOD] = 0;
-	NumUVChannels[CurrentLOD] = 0;
+	NumTriangles[CurrentLOD] = 0; //-V781
+	NumVertices[CurrentLOD] = 0; //-V781
+	NumUVChannels[CurrentLOD] = 0; //-V781
 	int32 NumLODLevels = 0;
 
 	if( StaticMesh->RenderData )
diff --git a/Engine/Source/Editor/UnrealEd/Private/EditorExporters.cpp b/Engine/Source/Editor/UnrealEd/Private/EditorExporters.cpp
--- a/Engine/Source/Editor/UnrealEd/Private/EditorExporters.cpp
+++ b/Engine/Source/Editor/UnrealEd/Private/EditorExporters.cpp
@@ -1657,7 +1657,8 @@ bool UPolysExporterOBJ::ExportText(const FExportObjectInnerContext* Context, UOb
 
 	Ar.Logf (TEXT("# OBJ File Generated by UnrealEd\n"));
 
-	ExportPolys( Polys, PolyNum, TotalPolys, Warn, false, NULL, Objects );
+	UWorld* World = CastChecked<UWorld>(Object);
+	ExportPolys( Polys, PolyNum, TotalPolys, Warn, false, World->GetModel(), Objects );
 
 	for( int32 o = 0 ; o < Objects.Num() ; ++o )
 	{
diff --git a/Engine/Source/Editor/VREditor/VREditorMotionControllerInteractor.cpp b/Engine/Source/Editor/VREditor/VREditorMotionControllerInteractor.cpp
--- a/Engine/Source/Editor/VREditor/VREditorMotionControllerInteractor.cpp
+++ b/Engine/Source/Editor/VREditor/VREditorMotionControllerInteractor.cpp
@@ -306,7 +306,7 @@ void UVREditorMotionControllerInteractor::SetupComponent( AActor* OwningActor )
 			{
 				StaticMesh = StartSplineMesh;
 			}
-			else if (i == NumLaserSplinePoints)
+			else if (i == NumLaserSplinePoints - 1)
 			{
 				StaticMesh = EndSplineMesh;
 			}
diff --git a/Engine/Source/Runtime/NetworkFile/Private/NetworkPlatformFile.cpp b/Engine/Source/Runtime/NetworkFile/Private/NetworkPlatformFile.cpp
--- a/Engine/Source/Runtime/NetworkFile/Private/NetworkPlatformFile.cpp
+++ b/Engine/Source/Runtime/NetworkFile/Private/NetworkPlatformFile.cpp
@@ -1145,7 +1145,7 @@ void FNetworkPlatformFile::EnsureFileIsLocal(const FString& Filename)
 	}
 	if (FinishedAsyncWriteUnsolicitedFiles)
 	{
-		if (bIncrimentedPackageWaits == false && FinishedAsyncNetworkReadUnsolicitedFiles->Get() == 0)
+		if (bIncrimentedPackageWaits == false && FinishedAsyncNetworkReadUnsolicitedFiles->Get() == 0) //-V522
 		{
 			++UnsolicitedPackageWaits;
 		}
