New IP Addr changes for Oculus

diff --git a/Engine/Plugins/Online/OnlineSubsystemOculus/Source/Private/IPAddressOculus.h b/Engine/Plugins/Online/OnlineSubsystemOculus/Source/Private/IPAddressOculus.h
--- a/Engine/Plugins/Online/OnlineSubsystemOculus/Source/Private/IPAddressOculus.h
+++ b/Engine/Plugins/Online/OnlineSubsystemOculus/Source/Private/IPAddressOculus.h
@@ -4,6 +4,7 @@
 
 #include "CoreMinimal.h"
 #include "IPAddress.h"
+#include "Algo/Reverse.h"
 #include "Engine/EngineBaseTypes.h"
 #include "OnlineSubsystemOculusTypes.h"
 #include "OnlineSubsystemOculusPackage.h"
@@ -59,6 +60,44 @@ public:
 		return OculusId.GetID();
 	}
 
+	virtual TArray<uint8> GetRawIp() const override
+	{
+		TArray<uint8> RawAddressArray;
+		const uint8* OculusIdWalk = OculusId.GetBytes();
+		while (RawAddressArray.Num() < OculusId.GetSize())
+		{
+			RawAddressArray.Add(*OculusIdWalk);
+			++OculusIdWalk;
+		}
+
+		// We want to make sure that these arrays are in big endian format.
+#if PLATFORM_LITTLE_ENDIAN
+		Algo::Reverse(RawAddressArray);
+#endif
+
+		return RawAddressArray;
+	}
+
+	virtual void SetRawIp(const TArray<uint8>& RawAddr) override
+	{
+		ovrID NewId;
+
+		// Make a quick copy of the array
+		TArray<uint8> WorkingArray = RawAddr;
+
+		// If we're little endian, we need to reverse so that the shifts turn out correct
+#if PLATFORM_LITTLE_ENDIAN
+		Algo::Reverse(WorkingArray);
+#endif
+
+		for (int32 i = 0; i < WorkingArray.Num(); ++i)
+		{
+			NewId |= (ovrID)WorkingArray[i] << (i * 8);
+		}
+
+		OculusId = FUniqueNetIdOculus(NewId);
+	}
+
 	/**
 	* Sets the ip address from a host byte order uint32
 	*
@@ -130,6 +169,12 @@ public:
 		/** Not used */
 	}
 
+	/** Sets the address to loopback */
+	void SetLoopbackAddress() override
+	{
+		/** Not used */
+	}
+
 	/**
 	* Converts this internet ip address to string form
 	*
@@ -170,4 +215,11 @@ public:
 	{
 		return OculusId.IsValid();
 	}
+
+	virtual TSharedRef<FInternetAddr> Clone() const override
+	{
+		TSharedRef<FInternetAddrOculus> NewAddress = MakeShareable(new FInternetAddrOculus);
+		NewAddress->OculusId = OculusId;
+		return NewAddress;
+	}
 };
