Fix issue caused by removal of TickHandle member from OnlineSubsystemImpl. It now inherits from FTickerObjectBase and is toggled on/off via a bool.

diff --git a/Engine/Plugins/Online/OnlineSubsystemOculus/Source/Private/OnlineSubsystemOculus.cpp b/Engine/Plugins/Online/OnlineSubsystemOculus/Source/Private/OnlineSubsystemOculus.cpp
--- a/Engine/Plugins/Online/OnlineSubsystemOculus/Source/Private/OnlineSubsystemOculus.cpp
+++ b/Engine/Plugins/Online/OnlineSubsystemOculus/Source/Private/OnlineSubsystemOculus.cpp
@@ -236,10 +236,7 @@ bool FOnlineSubsystemOculus::Init()
 		// Shutdown stops the ticker, but construction of the object starts the ticker.
 		// Since this hangs around, ensure that the ticker gets started in the editor when 
 		// we init.
-		if (!TickHandle.IsValid())
-		{
-			StartTicker();
-		}
+		StartTicker();
 #endif
 	}
 	else
