Fix missing call to 'ResetChallengeData' on some NetDrivers.

Resolves issue where platforms which have identical net addresses on reconnect, would - if they were the last to connect - end up with the old connections randomized sequence.

diff --git a/Engine/Plugins/Online/OnlineSubsystemOculus/Source/Private/OculusNetDriver.cpp b/Engine/Plugins/Online/OnlineSubsystemOculus/Source/Private/OculusNetDriver.cpp
--- a/Engine/Plugins/Online/OnlineSubsystemOculus/Source/Private/OculusNetDriver.cpp
+++ b/Engine/Plugins/Online/OnlineSubsystemOculus/Source/Private/OculusNetDriver.cpp
@@ -239,6 +239,8 @@ void UOculusNetDriver::TickDispatch(float DeltaTime)
 					StatelessConnect->GetChallengeSequence(ServerSequence, ClientSequence);
 
 					Connection->InitSequence(ClientSequence, ServerSequence);
+
+					StatelessConnect->ResetChallengeData();
 				}
 
 				if (Connection->Handler.IsValid())
