Fix crash with oculus room invites

diff --git a/Engine/Plugins/Online/OnlineSubsystemOculus/Source/Private/OnlineSessionInterfaceOculus.cpp b/Engine/Plugins/Online/OnlineSubsystemOculus/Source/Private/OnlineSessionInterfaceOculus.cpp
--- a/Engine/Plugins/Online/OnlineSubsystemOculus/Source/Private/OnlineSessionInterfaceOculus.cpp
+++ b/Engine/Plugins/Online/OnlineSubsystemOculus/Source/Private/OnlineSessionInterfaceOculus.cpp
@@ -1336,16 +1336,18 @@ void FOnlineSessionOculus::OnRoomInviteAccepted(ovrMessageHandle Message, bool b
 
 		auto Room = ovr_Message_GetRoom(InMessage);
 		auto Session = CreateSessionFromRoom(Room);
-		LocalSearchResult.Session = Session.Get();
 		
 		// check if there's a delegate bound, if not save this session for later.
 		if (!OnSessionUserInviteAcceptedDelegates.IsBound())
 		{
 			// No delegates are bound, just save this for later
-			PendingInviteAcceptedSessions.Add(MakeShareable(&LocalSearchResult));
+			TSharedRef<FOnlineSessionSearchResult> SearchResult = MakeShareable(new FOnlineSessionSearchResult());
+			SearchResult->Session = Session.Get();
+			PendingInviteAcceptedSessions.Add(SearchResult);
 			return;
 		}
 
+		LocalSearchResult.Session = Session.Get();
 		TriggerOnSessionUserInviteAcceptedDelegates(true, 0, PlayerId, LocalSearchResult);
 	}));
 }
