Fixed CIS errors on nightly build

#rnx

diff --git a/Engine/Plugins/Online/OnlineSubsystemOculus/Source/Private/OnlineSessionInterfaceOculus.cpp b/Engine/Plugins/Online/OnlineSubsystemOculus/Source/Private/OnlineSessionInterfaceOculus.cpp
--- a/Engine/Plugins/Online/OnlineSubsystemOculus/Source/Private/OnlineSessionInterfaceOculus.cpp
+++ b/Engine/Plugins/Online/OnlineSubsystemOculus/Source/Private/OnlineSessionInterfaceOculus.cpp
@@ -1341,9 +1341,9 @@ void FOnlineSessionOculus::OnRoomInviteAccepted(ovrMessageHandle Message, bool b
 		if (!OnSessionUserInviteAcceptedDelegates.IsBound())
 		{
 			// No delegates are bound, just save this for later
-			TSharedRef<FOnlineSessionSearchResult> SearchResult = MakeShareable(new FOnlineSessionSearchResult());
-			SearchResult->Session = Session.Get();
-			PendingInviteAcceptedSessions.Add(SearchResult);
+			TSharedRef<FOnlineSessionSearchResult> SearchResultPtr = MakeShareable(new FOnlineSessionSearchResult());
+			SearchResultPtr->Session = Session.Get();
+			PendingInviteAcceptedSessions.Add(SearchResultPtr);
 			return;
 		}
 
diff --git a/Engine/Source/Runtime/Core/Private/Lumin/LuminPlatformFile.cpp b/Engine/Source/Runtime/Core/Private/Lumin/LuminPlatformFile.cpp
--- a/Engine/Source/Runtime/Core/Private/Lumin/LuminPlatformFile.cpp
+++ b/Engine/Source/Runtime/Core/Private/Lumin/LuminPlatformFile.cpp
@@ -30,7 +30,7 @@ FString AndroidRelativeToAbsolutePath(bool bUseInternalBasePath, FString RelPath
 		do 
 		{
 			Result.RightChopInline(3, false);
-		} while (Result.StartsWith(TEXT("../"), ESearchCase::CaseSensitive))
+		} while (Result.StartsWith(TEXT("../"), ESearchCase::CaseSensitive));
 	}
 
 	// Remove the base app path if present, we will prepend it the correct base path as needed.
