Remove or edit references to UE4

diff --git a/Engine/Plugins/Online/OnlineSubsystemOculus/Source/Private/OVR_PlatformLoader.cpp b/Engine/Plugins/Online/OnlineSubsystemOculus/Source/Private/OVR_PlatformLoader.cpp
--- a/Engine/Plugins/Online/OnlineSubsystemOculus/Source/Private/OVR_PlatformLoader.cpp
+++ b/Engine/Plugins/Online/OnlineSubsystemOculus/Source/Private/OVR_PlatformLoader.cpp
@@ -653,7 +653,7 @@ OVRPL_PUBLIC_FUNCTION(void) ovr_PlatformInitializeStandaloneAccessToken(const ch
 }
 
 // Resetting to pre-init state without unloading completely
-// Useful for development purposes such as running it in the Unity or UE4 editor
+// Useful for development purposes such as running it in the Unity or UE editor
 OVRP_PUBLIC_FUNCTION(void) ovr_UnityResetTestPlatform();
 OVRPL_PUBLIC_FUNCTION(void) ovr_ResetInitAndContext()
 {
diff --git a/Engine/Plugins/Online/OnlineSubsystemOculus/Source/Private/OnlineSessionInterfaceOculus.cpp b/Engine/Plugins/Online/OnlineSubsystemOculus/Source/Private/OnlineSessionInterfaceOculus.cpp
--- a/Engine/Plugins/Online/OnlineSubsystemOculus/Source/Private/OnlineSessionInterfaceOculus.cpp
+++ b/Engine/Plugins/Online/OnlineSubsystemOculus/Source/Private/OnlineSessionInterfaceOculus.cpp
@@ -746,7 +746,7 @@ bool FOnlineSessionOculus::FindModeratedRoomSessions(const TSharedRef<FOnlineSes
 
 		SearchSettings->SearchResults.Reset(SearchResultsSize);
 
-		// UE4 specific setting.  Sessions with different build unique ids shouldn't be able to see each other
+		// Sessions with different build unique ids shouldn't be able to see each other
 		// because the builds are not compatible
 		int32 BuildUniqueId = GetBuildUniqueId();
 
@@ -812,7 +812,7 @@ bool FOnlineSessionOculus::FindMatchmakingSessions(const FString Pool, const TSh
 
 		SearchSettings->SearchResults.Reset(SearchResultsSize);
 
-		// UE4 specific setting.  Sessions with different build unique ids shouldn't be able to see each other
+		// Sessions with different build unique ids shouldn't be able to see each other
 		// because the builds are not compatible
 		int32 BuildUniqueId = GetBuildUniqueId();
 
@@ -879,7 +879,7 @@ bool FOnlineSessionOculus::FindSessionById(const FUniqueNetId& SearchingUserId,
 			return;
 		}
 
-		// UE4 specific setting.  Sessions with different build unique ids shouldn't be able to see each other
+		// Sessions with different build unique ids shouldn't be able to see each other
 		// because the builds are not compatible
 		int32 BuildUniqueId = GetBuildUniqueId();
 
@@ -1005,7 +1005,7 @@ bool FOnlineSessionOculus::FindFriendSession(int32 LocalUserNum, const FUniqueNe
 			return;
 		}
 	
-	// UE4 specific setting.  Sessions with different build unique ids shouldn't be able to see each other
+		// Sessions with different build unique ids shouldn't be able to see each other
 		// because the builds are not compatible
 		int32 BuildUniqueId = GetBuildUniqueId();
 
