- Fix some broken TArray<FTCHARToUTF8> occurrences where we were calling .Get() on the elements before the last call to .Emplace, and without calling .Reserve to guarantee no relocations, potentially running into bugs where the inline buffers in the converter were relocated due to the TArray growing.

diff --git a/Engine/Plugins/Online/OnlineSubsystemOculus/Source/Private/OnlineSessionInterfaceOculus.cpp b/Engine/Plugins/Online/OnlineSubsystemOculus/Source/Private/OnlineSessionInterfaceOculus.cpp
--- a/Engine/Plugins/Online/OnlineSubsystemOculus/Source/Private/OnlineSessionInterfaceOculus.cpp
+++ b/Engine/Plugins/Online/OnlineSubsystemOculus/Source/Private/OnlineSessionInterfaceOculus.cpp
@@ -419,8 +419,10 @@ bool FOnlineSessionOculus::UpdateRoomDataStore(FName SessionName, FOnlineSession
 	TArray<ovrKeyValuePair> DataStore;
 	// Add the updated settings that changed
 	// We need to keep an array of these in scope because ovrKeyValuePair_makeString() does not call strdup()
-	TArray<FTCHARToUTF8*> DataStoreKeys;
-	TArray<FTCHARToUTF8*> DataStoreValues;
+	TArray<FTCHARToUTF8> DataStoreUpdatedKeys;
+	TArray<FTCHARToUTF8> DataStoreUpdatedValues;
+	DataStoreUpdatedKeys.Reserve(UpdatedSessionSettings.Settings.Num());
+	DataStoreUpdatedValues.Reserve(UpdatedSessionSettings.Settings.Num());
 
 	for (auto& Setting : UpdatedSessionSettings.Settings)
 	{
@@ -432,20 +434,20 @@ bool FOnlineSessionOculus::UpdateRoomDataStore(FName SessionName, FOnlineSession
 		FOnlineSessionSetting* ExistingSetting = Session->SessionSettings.Settings.Find(Setting.Key);
 		if (!ExistingSetting || ExistingSetting->Data != Setting.Value.Data)
 		{
-			auto Key = new FTCHARToUTF8(*Setting.Key.ToString());
-			DataStoreKeys.Add(MoveTemp(Key));
+			DataStoreUpdatedKeys.Emplace(*Setting.Key.ToString());
 			// Always convert to a string because that's the only type supported for room data store
-			auto Value = new FTCHARToUTF8(*Setting.Value.Data.ToString());
-			DataStoreValues.Add(MoveTemp(Value));
+			DataStoreUpdatedValues.Emplace(*Setting.Value.Data.ToString());
 
 			DataStore.Add(ovrKeyValuePair_makeString(
-				DataStoreKeys.Last()->Get(),
-				DataStoreValues.Last()->Get()
+				DataStoreUpdatedKeys.Last().Get(),
+				DataStoreUpdatedValues.Last().Get()
 			));
 		}
 	}
 
 	// Clear existing keys that use to exist
+	TArray<FTCHARToUTF8> DataStoreDeletedKeys;
+	DataStoreDeletedKeys.Reserve(Session->SessionSettings.Settings.Num());
 	for (auto& Setting : Session->SessionSettings.Settings)
 	{
 		// Oculus matchmaking pool key cannot be added or changed
@@ -457,24 +459,16 @@ bool FOnlineSessionOculus::UpdateRoomDataStore(FName SessionName, FOnlineSession
 		FOnlineSessionSetting* ExistingSetting = UpdatedSessionSettings.Settings.Find(Setting.Key);
 		if (!ExistingSetting)
 		{
-			auto Key = new FTCHARToUTF8(*Setting.Key.ToString());
-			DataStoreKeys.Add(MoveTemp(Key));
-			DataStore.Add(ovrKeyValuePair_makeString(DataStoreKeys.Last()->Get(), ""));
+			DataStoreDeletedKeys.Emplace(*Setting.Key.ToString());
+			DataStore.Add(ovrKeyValuePair_makeString(DataStoreDeletedKeys.Last().Get(), ""));
 		}
 	}
-	auto DataStoreSize = DataStore.Num();
 
 	// If there is a delta, then fire off the request
-	if (DataStoreSize > 0)
+	if (DataStore.Num() > 0)
 	{
-		// move to a regular array to put in our OVR function
-		ovrKeyValuePair* DataStoreArray = new ovrKeyValuePair[DataStoreSize];
-		for (int i = 0; i < DataStoreSize; ++i)
-		{
-			DataStoreArray[i] = MoveTemp(DataStore[i]);
-		}
 		OculusSubsystem.AddRequestDelegate(
-			ovr_Room_UpdateDataStore(GetOvrIDFromSession(*Session), DataStoreArray, DataStoreSize),
+			ovr_Room_UpdateDataStore(GetOvrIDFromSession(*Session), DataStore.GetData(), DataStore.Num()),
 			FOculusMessageOnCompleteDelegate::CreateLambda([this, SessionName](ovrMessageHandle Message, bool bIsError)
 		{
 			if (bIsError)
@@ -500,12 +494,6 @@ bool FOnlineSessionOculus::UpdateRoomDataStore(FName SessionName, FOnlineSession
 
 			TriggerOnUpdateSessionCompleteDelegates(SessionName, true);
 		}));
-
-		// clean up the arrays
-		delete[] DataStoreArray;
-		DataStore.Empty();
-		DataStoreKeys.Empty();
-		DataStoreValues.Empty();
 	}
 	else
 	{
