Deprecate OCULUS_SUBSYSTEM

#tests QAGame and UnrealGame compiled locally and by preflight

diff --git a/Engine/Plugins/Online/OnlineSubsystem/Source/Private/OnlineSubsystem.cpp b/Engine/Plugins/Online/OnlineSubsystem/Source/Private/OnlineSubsystem.cpp
--- a/Engine/Plugins/Online/OnlineSubsystem/Source/Private/OnlineSubsystem.cpp
+++ b/Engine/Plugins/Online/OnlineSubsystem/Source/Private/OnlineSubsystem.cpp
@@ -63,6 +63,7 @@ PRAGMA_DISABLE_DEPRECATION_WARNINGS
 FName WECHAT_SUBSYSTEM(TEXT("WeChat"));
 FName LIVE_SUBSYSTEM(TEXT("LIVE"));
 FName LIVESERVER_SUBSYSTEM(TEXT("LIVESERVER"));
+FName OCULUS_SUBSYSTEM(TEXT("OCULUS"));
 PRAGMA_ENABLE_DEPRECATION_WARNINGS
 
 /** The default key that will update presence text in the platform's UI */
diff --git a/Engine/Plugins/Online/OnlineSubsystem/Source/Public/OnlineSubsystemNames.h b/Engine/Plugins/Online/OnlineSubsystem/Source/Public/OnlineSubsystemNames.h
--- a/Engine/Plugins/Online/OnlineSubsystem/Source/Public/OnlineSubsystemNames.h
+++ b/Engine/Plugins/Online/OnlineSubsystem/Source/Public/OnlineSubsystemNames.h
@@ -85,9 +85,7 @@ UE_DEPRECATED(4.27, "OnlineSubsystemLiveServer has been deprecated and will be r
 #define SWITCH_SUBSYSTEM FName(TEXT("SWITCH"))
 #endif
 
-#ifndef OCULUS_SUBSYSTEM
-#define OCULUS_SUBSYSTEM FName(TEXT("Oculus"))
-#endif
+UE_DEPRECATED(5.2, "OnlineSubsystemOculus has been deprecated and will be removed") ONLINESUBSYSTEM_API extern FName OCULUS_SUBSYSTEM;
 
 #ifndef SAMSUNG_SUBSYSTEM
 #define SAMSUNG_SUBSYSTEM FName(TEXT("SAMSUNG"))
diff --git a/Engine/Plugins/Online/OnlineSubsystemOculus/OnlineSubsystemOculus.uplugin b/Engine/Plugins/Online/OnlineSubsystemOculus/OnlineSubsystemOculus.uplugin
--- a/Engine/Plugins/Online/OnlineSubsystemOculus/OnlineSubsystemOculus.uplugin
+++ b/Engine/Plugins/Online/OnlineSubsystemOculus/OnlineSubsystemOculus.uplugin
@@ -3,7 +3,7 @@
 	"Version": 2,
 	"VersionName": "1.32.0",
 	"FriendlyName": "Online Subsystem Oculus",
-	"Description": "Access to Oculus platform",
+	"Description": "Access to Oculus platform.  This plugin is deprecated and will be removed in a future engine release.",
 	"Category": "Online Platform",
 	"CreatedBy": "Oculus",
 	"CreatedByURL": "https://www.oculus.com/",
diff --git a/Engine/Plugins/Online/OnlineSubsystemOculus/Source/Private/OculusCreateSessionCallbackProxy.cpp b/Engine/Plugins/Online/OnlineSubsystemOculus/Source/Private/OculusCreateSessionCallbackProxy.cpp
--- a/Engine/Plugins/Online/OnlineSubsystemOculus/Source/Private/OculusCreateSessionCallbackProxy.cpp
+++ b/Engine/Plugins/Online/OnlineSubsystemOculus/Source/Private/OculusCreateSessionCallbackProxy.cpp
@@ -25,8 +25,9 @@ UOculusCreateSessionCallbackProxy* UOculusCreateSessionCallbackProxy::CreateSess
 
 void UOculusCreateSessionCallbackProxy::Activate()
 {
+PRAGMA_DISABLE_DEPRECATION_WARNINGS
 	auto OculusSessionInterface = Online::GetSessionInterface(OCULUS_SUBSYSTEM);
-
+PRAGMA_ENABLE_DEPRECATION_WARNINGS
 	if (OculusSessionInterface.IsValid())
 	{
 		CreateCompleteDelegateHandle = OculusSessionInterface->AddOnCreateSessionCompleteDelegate_Handle(CreateCompleteDelegate);
@@ -54,7 +55,9 @@ void UOculusCreateSessionCallbackProxy::Activate()
 
 void UOculusCreateSessionCallbackProxy::OnCreateCompleted(FName SessionName, bool bWasSuccessful)
 {
+PRAGMA_DISABLE_DEPRECATION_WARNINGS
 	auto OculusSessionInterface = Online::GetSessionInterface(OCULUS_SUBSYSTEM);
+PRAGMA_ENABLE_DEPRECATION_WARNINGS
 
 	if (OculusSessionInterface.IsValid())
 	{
@@ -78,7 +81,9 @@ void UOculusCreateSessionCallbackProxy::OnCreateCompleted(FName SessionName, boo
 
 void UOculusCreateSessionCallbackProxy::OnStartCompleted(FName SessionName, bool bWasSuccessful)
 {
+PRAGMA_DISABLE_DEPRECATION_WARNINGS
 	auto OculusSessionInterface = Online::GetSessionInterface(OCULUS_SUBSYSTEM);
+PRAGMA_ENABLE_DEPRECATION_WARNINGS
 
 	if (OculusSessionInterface.IsValid())
 	{
diff --git a/Engine/Plugins/Online/OnlineSubsystemOculus/Source/Private/OculusEntitlementCallbackProxy.cpp b/Engine/Plugins/Online/OnlineSubsystemOculus/Source/Private/OculusEntitlementCallbackProxy.cpp
--- a/Engine/Plugins/Online/OnlineSubsystemOculus/Source/Private/OculusEntitlementCallbackProxy.cpp
+++ b/Engine/Plugins/Online/OnlineSubsystemOculus/Source/Private/OculusEntitlementCallbackProxy.cpp
@@ -19,7 +19,9 @@ UOculusEntitlementCallbackProxy* UOculusEntitlementCallbackProxy::VerifyEntitlem
 
 void UOculusEntitlementCallbackProxy::Activate()
 {
+PRAGMA_DISABLE_DEPRECATION_WARNINGS
 	auto OculusIdentityInterface = Online::GetIdentityInterface(OCULUS_SUBSYSTEM);
+PRAGMA_ENABLE_DEPRECATION_WARNINGS
 
 	if (OculusIdentityInterface.IsValid())
 	{
diff --git a/Engine/Plugins/Online/OnlineSubsystemOculus/Source/Private/OculusFindSessionsCallbackProxy.cpp b/Engine/Plugins/Online/OnlineSubsystemOculus/Source/Private/OculusFindSessionsCallbackProxy.cpp
--- a/Engine/Plugins/Online/OnlineSubsystemOculus/Source/Private/OculusFindSessionsCallbackProxy.cpp
+++ b/Engine/Plugins/Online/OnlineSubsystemOculus/Source/Private/OculusFindSessionsCallbackProxy.cpp
@@ -35,7 +35,9 @@ UOculusFindSessionsCallbackProxy* UOculusFindSessionsCallbackProxy::FindModerate
 
 void UOculusFindSessionsCallbackProxy::Activate()
 {
+PRAGMA_DISABLE_DEPRECATION_WARNINGS
 	auto OculusSessionInterface = Online::GetSessionInterface(OCULUS_SUBSYSTEM);
+PRAGMA_ENABLE_DEPRECATION_WARNINGS
 
 	if (OculusSessionInterface.IsValid())
 	{
@@ -62,7 +64,9 @@ void UOculusFindSessionsCallbackProxy::Activate()
 
 void UOculusFindSessionsCallbackProxy::OnCompleted(bool bSuccess)
 {
+PRAGMA_DISABLE_DEPRECATION_WARNINGS
 	auto OculusSessionInterface = Online::GetSessionInterface(OCULUS_SUBSYSTEM);
+PRAGMA_ENABLE_DEPRECATION_WARNINGS
 
 	if (OculusSessionInterface.IsValid())
 	{
diff --git a/Engine/Plugins/Online/OnlineSubsystemOculus/Source/Private/OculusIdentityCallbackProxy.cpp b/Engine/Plugins/Online/OnlineSubsystemOculus/Source/Private/OculusIdentityCallbackProxy.cpp
--- a/Engine/Plugins/Online/OnlineSubsystemOculus/Source/Private/OculusIdentityCallbackProxy.cpp
+++ b/Engine/Plugins/Online/OnlineSubsystemOculus/Source/Private/OculusIdentityCallbackProxy.cpp
@@ -19,7 +19,9 @@ UOculusIdentityCallbackProxy* UOculusIdentityCallbackProxy::GetOculusIdentity(in
 
 void UOculusIdentityCallbackProxy::Activate()
 {
+PRAGMA_DISABLE_DEPRECATION_WARNINGS
 	auto OculusIdentityInterface = Online::GetIdentityInterface(OCULUS_SUBSYSTEM);
+PRAGMA_ENABLE_DEPRECATION_WARNINGS
 
 	if (OculusIdentityInterface.IsValid())
 	{
diff --git a/Engine/Plugins/Online/OnlineSubsystemOculus/Source/Private/OculusNetDriver.cpp b/Engine/Plugins/Online/OnlineSubsystemOculus/Source/Private/OculusNetDriver.cpp
--- a/Engine/Plugins/Online/OnlineSubsystemOculus/Source/Private/OculusNetDriver.cpp
+++ b/Engine/Plugins/Online/OnlineSubsystemOculus/Source/Private/OculusNetDriver.cpp
@@ -13,8 +13,10 @@
 
 bool UOculusNetDriver::IsAvailable() const
 {
+PRAGMA_DISABLE_DEPRECATION_WARNINGS
 	// Net driver won't work if the online subsystem doesn't exist
 	IOnlineSubsystem* OculusSubsystem = IOnlineSubsystem::Get(OCULUS_SUBSYSTEM);
+PRAGMA_ENABLE_DEPRECATION_WARNINGS
 	if (OculusSubsystem)
 	{
 		return true;
@@ -83,11 +85,13 @@ bool UOculusNetDriver::InitBase(bool bInitAsClient, FNetworkNotify* InNotify, co
 	// Listen for network state
 	if (!NetworkingConnectionStateChangeDelegateHandle.IsValid())
 	{
+PRAGMA_DISABLE_DEPRECATION_WARNINGS
 		auto OnlineSubsystem = static_cast<FOnlineSubsystemOculus*>(IOnlineSubsystem::Get(OCULUS_SUBSYSTEM));
 		NetworkingConnectionStateChangeDelegateHandle =
 			OnlineSubsystem->GetNotifDelegate(ovrMessage_Notification_Networking_ConnectionStateChange)
 			.AddUObject(this, &UOculusNetDriver::OnNetworkingConnectionStateChange);
 	}
+PRAGMA_ENABLE_DEPRECATION_WARNINGS
 
 	return true;
 }
@@ -145,11 +149,13 @@ bool UOculusNetDriver::InitListen(FNetworkNotify* InNotify, FURL& LocalURL, bool
 	// Listen for incoming peers
 	if (!PeerConnectRequestDelegateHandle.IsValid())
 	{
+PRAGMA_DISABLE_DEPRECATION_WARNINGS
 		auto OnlineSubsystem = static_cast<FOnlineSubsystemOculus*>(IOnlineSubsystem::Get(OCULUS_SUBSYSTEM));
 		PeerConnectRequestDelegateHandle =
 			OnlineSubsystem->GetNotifDelegate(ovrMessage_Notification_Networking_PeerConnectRequest)
 			.AddUObject(this, &UOculusNetDriver::OnNewNetworkingPeerRequest);
 	}
+PRAGMA_ENABLE_DEPRECATION_WARNINGS
 
 	InitConnectionlessHandler();
 
@@ -470,7 +476,7 @@ void UOculusNetDriver::Shutdown()
 	UNetDriver::Shutdown();
 
 	UE_LOG(LogNet, Verbose, TEXT("Oculus Net Driver shutdown"));
-
+PRAGMA_DISABLE_DEPRECATION_WARNINGS
 	auto OnlineSubsystem = static_cast<FOnlineSubsystemOculus*>(IOnlineSubsystem::Get(OCULUS_SUBSYSTEM));
 	if (PeerConnectRequestDelegateHandle.IsValid())
 	{
@@ -482,6 +488,7 @@ void UOculusNetDriver::Shutdown()
 		OnlineSubsystem->RemoveNotifDelegate(ovrMessage_Notification_Networking_ConnectionStateChange, NetworkingConnectionStateChangeDelegateHandle);
 		NetworkingConnectionStateChangeDelegateHandle.Reset();
 	}
+PRAGMA_ENABLE_DEPRECATION_WARNINGS
 
 	// Ensure all current connections are closed now
 	for (auto& Connection : Connections)
diff --git a/Engine/Plugins/Online/OnlineSubsystemOculus/Source/Private/OculusUpdateSessionCallbackProxy.cpp b/Engine/Plugins/Online/OnlineSubsystemOculus/Source/Private/OculusUpdateSessionCallbackProxy.cpp
--- a/Engine/Plugins/Online/OnlineSubsystemOculus/Source/Private/OculusUpdateSessionCallbackProxy.cpp
+++ b/Engine/Plugins/Online/OnlineSubsystemOculus/Source/Private/OculusUpdateSessionCallbackProxy.cpp
@@ -23,7 +23,9 @@ UOculusUpdateSessionCallbackProxy* UOculusUpdateSessionCallbackProxy::SetSession
 
 void UOculusUpdateSessionCallbackProxy::Activate()
 {
+PRAGMA_DISABLE_DEPRECATION_WARNINGS
 	auto OculusSessionInterface = Online::GetSessionInterface(OCULUS_SUBSYSTEM);
+PRAGMA_ENABLE_DEPRECATION_WARNINGS
 
 	if (OculusSessionInterface.IsValid())
 	{
@@ -42,7 +44,9 @@ void UOculusUpdateSessionCallbackProxy::Activate()
 
 void UOculusUpdateSessionCallbackProxy::OnUpdateCompleted(FName SessionName, bool bWasSuccessful)
 {
+PRAGMA_DISABLE_DEPRECATION_WARNINGS
 	auto OculusSessionInterface = Online::GetSessionInterface(OCULUS_SUBSYSTEM);
+PRAGMA_ENABLE_DEPRECATION_WARNINGS
 
 	if (OculusSessionInterface.IsValid())
 	{
diff --git a/Engine/Plugins/Online/OnlineSubsystemOculus/Source/Private/OnlineAchievementsInterfaceOculus.cpp b/Engine/Plugins/Online/OnlineSubsystemOculus/Source/Private/OnlineAchievementsInterfaceOculus.cpp
--- a/Engine/Plugins/Online/OnlineSubsystemOculus/Source/Private/OnlineAchievementsInterfaceOculus.cpp
+++ b/Engine/Plugins/Online/OnlineSubsystemOculus/Source/Private/OnlineAchievementsInterfaceOculus.cpp
@@ -14,7 +14,7 @@ private:
 	FUniqueNetIdOculus PlayerId;
 	FOnlineAchievementsWriteRef WriteObject;
 	FOnAchievementsWrittenDelegate AchievementDelegate;
-
+PRAGMA_DISABLE_DEPRECATION_WARNINGS
 	// private to force the use of FOnlineMessageMultiTaskOculusWriteAchievements::Create()
 	FOnlineMessageMultiTaskOculusWriteAchievements(FOnlineSubsystemOculus& InOculusSubsystem, const FUniqueNetIdOculus& InPlayerId, FOnlineAchievementsWriteRef& InWriteObject, const FOnAchievementsWrittenDelegate& InAchievementDelegate)
 		: FOnlineMessageMultiTaskOculus(InOculusSubsystem, FOnlineMessageMultiTaskOculus::FFinalizeDelegate::CreateRaw(this, &FOnlineMessageMultiTaskOculusWriteAchievements::Finalize))
@@ -22,11 +22,12 @@ private:
 		, WriteObject(InWriteObject)
 		, AchievementDelegate(InAchievementDelegate)
 	{}
+PRAGMA_ENABLE_DEPRECATION_WARNINGS
 
 	static TSet< TSharedRef<FOnlineMessageMultiTaskOculusWriteAchievements> > ActiveAchievementWriteTasks;
 
 PACKAGE_SCOPE:
-
+PRAGMA_DISABLE_DEPRECATION_WARNINGS
 	static TSharedRef<FOnlineMessageMultiTaskOculusWriteAchievements> Create(
 		FOnlineSubsystemOculus& InOculusSubsystem,
 		const FUniqueNetIdOculus& InPlayerId,
@@ -38,6 +39,7 @@ PACKAGE_SCOPE:
 
 		return NewTask;
 	}
+PRAGMA_ENABLE_DEPRECATION_WARNINGS
 
 	void Finalize()
 	{
@@ -74,8 +76,9 @@ void FOnlineAchievementsOculus::WriteAchievements(const FUniqueNetId& PlayerId,
 		Delegate.ExecuteIfBound(PlayerId, false);
 		return;
 	}
-
+PRAGMA_DISABLE_DEPRECATION_WARNINGS
 	auto LoggedInPlayerId = OculusSubsystem.GetIdentityInterface()->GetUniquePlayerId(0);
+PRAGMA_ENABLE_DEPRECATION_WARNINGS
 	if (!(LoggedInPlayerId.IsValid() && PlayerId == *LoggedInPlayerId))
 	{
 		UE_LOG_ONLINE_ACHIEVEMENTS(Error, TEXT("Can only write achievements for logged in player id"));
@@ -150,7 +153,9 @@ void FOnlineAchievementsOculus::WriteAchievements(const FUniqueNetId& PlayerId,
 
 void FOnlineAchievementsOculus::QueryAchievements(const FUniqueNetId& PlayerId, const FOnQueryAchievementsCompleteDelegate& Delegate)
 {
+PRAGMA_DISABLE_DEPRECATION_WARNINGS
 	auto LoggedInPlayerId = OculusSubsystem.GetIdentityInterface()->GetUniquePlayerId(0);
+PRAGMA_ENABLE_DEPRECATION_WARNINGS
 	if (!(LoggedInPlayerId.IsValid() && PlayerId == *LoggedInPlayerId))
 	{
 		UE_LOG_ONLINE_ACHIEVEMENTS(Error, TEXT("Can only query for logged in player id"));
@@ -159,6 +164,7 @@ void FOnlineAchievementsOculus::QueryAchievements(const FUniqueNetId& PlayerId,
 	}
 
 	const FUniqueNetIdOculus& OculusPlayerId = FUniqueNetIdOculus::Cast(PlayerId);
+PRAGMA_DISABLE_DEPRECATION_WARNINGS
 	OculusSubsystem.AddRequestDelegate(
 		ovr_Achievements_GetAllProgress(),
 		FOculusMessageOnCompleteDelegate::CreateLambda([this, OculusPlayerId, Delegate](ovrMessageHandle Message, bool bIsError)
@@ -205,11 +211,13 @@ void FOnlineAchievementsOculus::QueryAchievements(const FUniqueNetId& PlayerId,
 
 		Delegate.ExecuteIfBound(OculusPlayerId, true);
 	}));
+PRAGMA_ENABLE_DEPRECATION_WARNINGS
 }
 
 void FOnlineAchievementsOculus::QueryAchievementDescriptions(const FUniqueNetId& PlayerId, const FOnQueryAchievementsCompleteDelegate& Delegate)
 {
 	const FUniqueNetIdOculus& OculusPlayerId = FUniqueNetIdOculus::Cast(PlayerId);
+PRAGMA_DISABLE_DEPRECATION_WARNINGS
 	OculusSubsystem.AddRequestDelegate(
 		ovr_Achievements_GetAllDefinitions(),
 		FOculusMessageOnCompleteDelegate::CreateLambda([this, OculusPlayerId, Delegate](ovrMessageHandle Message, bool bIsError)
@@ -239,6 +247,7 @@ void FOnlineAchievementsOculus::QueryAchievementDescriptions(const FUniqueNetId&
 
 		Delegate.ExecuteIfBound(OculusPlayerId, true);
 	}));
+PRAGMA_ENABLE_DEPRECATION_WARNINGS
 }
 
 EOnlineCachedResult::Type FOnlineAchievementsOculus::GetCachedAchievement(const FUniqueNetId& PlayerId, const FString& AchievementId, FOnlineAchievement& OutAchievement)
diff --git a/Engine/Plugins/Online/OnlineSubsystemOculus/Source/Private/OnlineAchievementsInterfaceOculus.h b/Engine/Plugins/Online/OnlineSubsystemOculus/Source/Private/OnlineAchievementsInterfaceOculus.h
--- a/Engine/Plugins/Online/OnlineSubsystemOculus/Source/Private/OnlineAchievementsInterfaceOculus.h
+++ b/Engine/Plugins/Online/OnlineSubsystemOculus/Source/Private/OnlineAchievementsInterfaceOculus.h
@@ -79,9 +79,11 @@ struct FOnlineAchievementOculus : FOnlineAchievement
 class FOnlineAchievementsOculus : public IOnlineAchievements
 {
 private:
-	
+
+PRAGMA_DISABLE_DEPRECATION_WARNINGS	
 	/** Reference to the owning subsystem */
 	FOnlineSubsystemOculus& OculusSubsystem;
+PRAGMA_ENABLE_DEPRECATION_WARNINGS
 
 	/** Mapping of players to their achievements */
 	TMap<FUniqueNetIdOculus, TArray<FOnlineAchievement>> PlayerAchievements;
@@ -95,12 +97,14 @@ private:
 
 public:
 
+PRAGMA_DISABLE_DEPRECATION_WARNINGS
 	/**
 	* Constructor
 	*
 	* @param InSubsystem - A reference to the owning subsystem
 	*/
 	FOnlineAchievementsOculus(FOnlineSubsystemOculus& InSubsystem);
+PRAGMA_ENABLE_DEPRECATION_WARNINGS
 
 	/**
 	* Default destructor
diff --git a/Engine/Plugins/Online/OnlineSubsystemOculus/Source/Private/OnlineFriendsInterfaceOculus.cpp b/Engine/Plugins/Online/OnlineSubsystemOculus/Source/Private/OnlineFriendsInterfaceOculus.cpp
--- a/Engine/Plugins/Online/OnlineSubsystemOculus/Source/Private/OnlineFriendsInterfaceOculus.cpp
+++ b/Engine/Plugins/Online/OnlineSubsystemOculus/Source/Private/OnlineFriendsInterfaceOculus.cpp
@@ -5,15 +5,17 @@
 #include "OnlineError.h"
 #include "OnlineSubsystemOculusPackage.h"
 
+PRAGMA_DISABLE_DEPRECATION_WARNINGS
 FOnlineFriendsOculus::FOnlineFriendsOculus(FOnlineSubsystemOculus& InSubsystem)
 : OculusSubsystem(InSubsystem)
 {
 }
-
+PRAGMA_ENABLE_DEPRECATION_WARNINGS
 const FString FOnlineFriendsOculus::FriendsListInviteableUsers = TEXT("invitableUsers");
 
 bool FOnlineFriendsOculus::ReadFriendsList(int32 LocalUserNum, const FString& ListName, const FOnReadFriendsListComplete& Delegate)
 {
+PRAGMA_DISABLE_DEPRECATION_WARNINGS
 	if (ListName == EFriendsLists::ToString(EFriendsLists::Default) || ListName == EFriendsLists::ToString(EFriendsLists::OnlinePlayers))
 	{
 		OculusSubsystem.AddRequestDelegate(
@@ -35,7 +37,7 @@ bool FOnlineFriendsOculus::ReadFriendsList(int32 LocalUserNum, const FString& Li
 		}));
 		return true;
 	}
-
+PRAGMA_ENABLE_DEPRECATION_WARNINGS
 	Delegate.ExecuteIfBound(LocalUserNum, false, ListName, TEXT("Invalid friends list"));
 	return false;
 }
@@ -77,12 +79,14 @@ void FOnlineFriendsOculus::OnQueryFriendsComplete(ovrMessageHandle Message, bool
 	bool bHasPaging = ovr_UserArray_HasNextPage(UserArray);
 	if (bHasPaging)
 	{
+PRAGMA_DISABLE_DEPRECATION_WARNINGS
 		OculusSubsystem.AddRequestDelegate(
 			ovr_User_GetNextUserArrayPage(UserArray),
 			FOculusMessageOnCompleteDelegate::CreateLambda([this, LocalUserNum, ListName, &OutList, Delegate](ovrMessageHandle InMessage, bool bInIsError)
 		{
 			OnQueryFriendsComplete(InMessage, bInIsError, LocalUserNum, ListName, OutList, /* bAppendToExistingMap */ true, Delegate);
 		}));
+PRAGMA_ENABLE_DEPRECATION_WARNINGS
 	}
 	else
 	{
diff --git a/Engine/Plugins/Online/OnlineSubsystemOculus/Source/Private/OnlineFriendsInterfaceOculus.h b/Engine/Plugins/Online/OnlineSubsystemOculus/Source/Private/OnlineFriendsInterfaceOculus.h
--- a/Engine/Plugins/Online/OnlineSubsystemOculus/Source/Private/OnlineFriendsInterfaceOculus.h
+++ b/Engine/Plugins/Online/OnlineSubsystemOculus/Source/Private/OnlineFriendsInterfaceOculus.h
@@ -70,9 +70,11 @@ public:
 class FOnlineFriendsOculus : public IOnlineFriends
 {
 private:
-	
+
+PRAGMA_DISABLE_DEPRECATION_WARNINGS
 	/** Reference to the owning subsystem */
 	FOnlineSubsystemOculus& OculusSubsystem;
+PRAGMA_ENABLE_DEPRECATION_WARNINGS
 
 	/** All the friends for the player */
 	TMap<uint64, TSharedRef<FOnlineFriend>> PlayerFriends;
@@ -92,12 +94,14 @@ public:
 
 	static const FString FriendsListInviteableUsers;
 
+PRAGMA_DISABLE_DEPRECATION_WARNINGS
 	/**
 	* Constructor
 	*
 	* @param InSubsystem - A reference to the owning subsystem
 	*/
 	FOnlineFriendsOculus(FOnlineSubsystemOculus& InSubsystem);
+PRAGMA_ENABLE_DEPRECATION_WARNINGS
 
 	/**
 	* Default destructor
diff --git a/Engine/Plugins/Online/OnlineSubsystemOculus/Source/Private/OnlineIdentityOculus.cpp b/Engine/Plugins/Online/OnlineSubsystemOculus/Source/Private/OnlineIdentityOculus.cpp
--- a/Engine/Plugins/Online/OnlineSubsystemOculus/Source/Private/OnlineIdentityOculus.cpp
+++ b/Engine/Plugins/Online/OnlineSubsystemOculus/Source/Private/OnlineIdentityOculus.cpp
@@ -58,9 +58,11 @@ bool FOnlineIdentityOculus::Login(int32 LocalUserNum, const FOnlineAccountCreden
 			{
 				// Immediately add the Oculus ID to our cache list
 				UserIds.Add(LocalUserNum, FUniqueNetIdOculus::Create(OculusId));
+PRAGMA_DISABLE_DEPRECATION_WARNINGS
 				OculusSubsystem.AddRequestDelegate(
 					ovr_User_GetLoggedInUser(),
 					FOculusMessageOnCompleteDelegate::CreateRaw(this, &FOnlineIdentityOculus::OnLoginComplete, LocalUserNum));
+PRAGMA_ENABLE_DEPRECATION_WARNINGS
 				return true;
 			}
 		}
@@ -253,15 +255,18 @@ void FOnlineIdentityOculus::RevokeAuthToken(const FUniqueNetId& UserId, const FO
 	});
 }
 
+PRAGMA_DISABLE_DEPRECATION_WARNINGS
 FOnlineIdentityOculus::FOnlineIdentityOculus(FOnlineSubsystemOculus& InSubsystem)
 	: OculusSubsystem(InSubsystem)
 {
 	// Auto login the 0-th player
 	FOnlineIdentityOculus::AutoLogin(0);
 }
+PRAGMA_ENABLE_DEPRECATION_WARNINGS
 
 void FOnlineIdentityOculus::GetUserPrivilege(const FUniqueNetId& UserId, EUserPrivileges::Type Privilege, const FOnGetUserPrivilegeCompleteDelegate& Delegate)
 {
+PRAGMA_DISABLE_DEPRECATION_WARNINGS
 	// Check for entitlement
 	OculusSubsystem.AddRequestDelegate(
 		ovr_Entitlement_GetIsViewerEntitled(),
@@ -284,6 +289,7 @@ void FOnlineIdentityOculus::GetUserPrivilege(const FUniqueNetId& UserId, EUserPr
 			}
 			Delegate.ExecuteIfBound(*UserId, Privilege, PrivilegeResults);
 		}));
+PRAGMA_ENABLE_DEPRECATION_WARNINGS
 }
 
 FPlatformUserId FOnlineIdentityOculus::GetPlatformUserIdFromUniqueNetId(const FUniqueNetId& UniqueNetId) const
diff --git a/Engine/Plugins/Online/OnlineSubsystemOculus/Source/Private/OnlineIdentityOculus.h b/Engine/Plugins/Online/OnlineSubsystemOculus/Source/Private/OnlineIdentityOculus.h
--- a/Engine/Plugins/Online/OnlineSubsystemOculus/Source/Private/OnlineIdentityOculus.h
+++ b/Engine/Plugins/Online/OnlineSubsystemOculus/Source/Private/OnlineIdentityOculus.h
@@ -82,6 +82,7 @@ public:
 
 	// FOnlineIdentityOculus
 
+PRAGMA_DISABLE_DEPRECATION_WARNINGS
 	/**
 	* Constructor
 	*
@@ -93,6 +94,7 @@ public:
 	* Default destructor
 	*/
 	virtual ~FOnlineIdentityOculus() = default;
+PRAGMA_ENABLE_DEPRECATION_WARNINGS
 
 PACKAGE_SCOPE:
 
@@ -104,9 +106,11 @@ PACKAGE_SCOPE:
 
 private:
 
+PRAGMA_DISABLE_DEPRECATION_WARNINGS
 	/** Reference to the owning subsystem */
 	FOnlineSubsystemOculus& OculusSubsystem;
-	
+PRAGMA_ENABLE_DEPRECATION_WARNINGS
+
 	/** Ids mapped to locally registered users */
 	TMap<int32, FUniqueNetIdRef> UserIds;
 
diff --git a/Engine/Plugins/Online/OnlineSubsystemOculus/Source/Private/OnlineLeaderboardInterfaceOculus.cpp b/Engine/Plugins/Online/OnlineSubsystemOculus/Source/Private/OnlineLeaderboardInterfaceOculus.cpp
--- a/Engine/Plugins/Online/OnlineSubsystemOculus/Source/Private/OnlineLeaderboardInterfaceOculus.cpp
+++ b/Engine/Plugins/Online/OnlineSubsystemOculus/Source/Private/OnlineLeaderboardInterfaceOculus.cpp
@@ -15,7 +15,9 @@ bool FOnlineLeaderboardOculus::ReadLeaderboards(const TArray< FUniqueNetIdRef >&
 	bool bOnlyLoggedInUser = false;
 	if (Players.Num() > 0) 
 	{
+PRAGMA_DISABLE_DEPRECATION_WARNINGS
 		auto LoggedInPlayerId = OculusSubsystem.GetIdentityInterface()->GetUniquePlayerId(0);
+PRAGMA_ENABLE_DEPRECATION_WARNINGS
 		if (Players.Num() == 1 && LoggedInPlayerId.IsValid() && *Players[0] == *LoggedInPlayerId)
 		{
 			bOnlyLoggedInUser = true;
@@ -60,12 +62,14 @@ bool FOnlineLeaderboardOculus::ReadOculusLeaderboards(bool bOnlyFriends, bool bO
 	}
 
 	ReadObject->ReadState = EOnlineAsyncTaskState::InProgress;
+PRAGMA_DISABLE_DEPRECATION_WARNINGS
 	OculusSubsystem.AddRequestDelegate(
 		ovr_Leaderboard_GetEntries(TCHAR_TO_ANSI(*ReadObject->LeaderboardName.ToString()), Limit, FilterType, StartAt),
 		FOculusMessageOnCompleteDelegate::CreateLambda([this, ReadObject](ovrMessageHandle Message, bool bIsError)
 	{
 		OnReadLeaderboardsComplete(Message, bIsError, ReadObject);
 	}));
+PRAGMA_ENABLE_DEPRECATION_WARNINGS
 	return true;
 }
 
@@ -151,12 +155,14 @@ void FOnlineLeaderboardOculus::OnReadLeaderboardsComplete(ovrMessageHandle Messa
 	bool bHasPaging = ovr_LeaderboardEntryArray_HasNextPage(LeaderboardArray);
 	if (bHasPaging)
 	{
+PRAGMA_DISABLE_DEPRECATION_WARNINGS
 		OculusSubsystem.AddRequestDelegate(
 			ovr_Leaderboard_GetNextEntries(LeaderboardArray),
 			FOculusMessageOnCompleteDelegate::CreateLambda([this, ReadObject](ovrMessageHandle InMessage, bool bInIsError)
 		{
 			OnReadLeaderboardsComplete(InMessage, bInIsError, ReadObject);
 		}));
+PRAGMA_ENABLE_DEPRECATION_WARNINGS
 		return;
 	}
 
@@ -171,7 +177,9 @@ void FOnlineLeaderboardOculus::FreeStats(FOnlineLeaderboardRead& ReadObject)
 
 bool FOnlineLeaderboardOculus::WriteLeaderboards(const FName& SessionName, const FUniqueNetId& Player, FOnlineLeaderboardWrite& WriteObject)
 {
+PRAGMA_DISABLE_DEPRECATION_WARNINGS
 	auto LoggedInPlayerId = OculusSubsystem.GetIdentityInterface()->GetUniquePlayerId(0);
+PRAGMA_ENABLE_DEPRECATION_WARNINGS
 	if (!(LoggedInPlayerId.IsValid() && Player == *LoggedInPlayerId))
 	{
 		UE_LOG_ONLINE_LEADERBOARD(Error, TEXT("Can only write to leaderboards for logged in player id"));
@@ -223,6 +231,7 @@ bool FOnlineLeaderboardOculus::WriteLeaderboards(const FName& SessionName, const
 
 	for (const auto& LeaderboardName : WriteObject.LeaderboardNames)
 	{
+PRAGMA_DISABLE_DEPRECATION_WARNINGS
 		OculusSubsystem.AddRequestDelegate(
 			ovr_Leaderboard_WriteEntry(TCHAR_TO_ANSI(*LeaderboardName.ToString()), Score, /* extra_data */ nullptr, 0, (WriteObject.UpdateMethod == ELeaderboardUpdateMethod::Force)),
 			FOculusMessageOnCompleteDelegate::CreateLambda([this](ovrMessageHandle Message, bool bIsError)
@@ -234,6 +243,7 @@ bool FOnlineLeaderboardOculus::WriteLeaderboards(const FName& SessionName, const
 				UE_LOG_ONLINE_LEADERBOARD(Error, TEXT("%s"), *FString(ErrorMessage));
 			}
 		}));
+PRAGMA_ENABLE_DEPRECATION_WARNINGS
 	}
 
 	return true;
diff --git a/Engine/Plugins/Online/OnlineSubsystemOculus/Source/Private/OnlineLeaderboardInterfaceOculus.h b/Engine/Plugins/Online/OnlineSubsystemOculus/Source/Private/OnlineLeaderboardInterfaceOculus.h
--- a/Engine/Plugins/Online/OnlineSubsystemOculus/Source/Private/OnlineLeaderboardInterfaceOculus.h
+++ b/Engine/Plugins/Online/OnlineSubsystemOculus/Source/Private/OnlineLeaderboardInterfaceOculus.h
@@ -12,21 +12,24 @@
 class FOnlineLeaderboardOculus : public IOnlineLeaderboards
 {
 private:
-	
+PRAGMA_DISABLE_DEPRECATION_WARNINGS
 	/** Reference to the owning subsystem */
 	FOnlineSubsystemOculus& OculusSubsystem;
+PRAGMA_ENABLE_DEPRECATION_WARNINGS
 
 	bool ReadOculusLeaderboards(bool bOnlyFriends, bool bOnlyLoggedInUser, FOnlineLeaderboardReadRef& ReadObject);
 	void OnReadLeaderboardsComplete(ovrMessageHandle Message, bool bIsError, const FOnlineLeaderboardReadRef& ReadObject);
 
 public:
 
+PRAGMA_DISABLE_DEPRECATION_WARNINGS
 	/**
 	* Constructor
 	*
 	* @param InSubsystem - A reference to the owning subsystem
 	*/
 	FOnlineLeaderboardOculus(FOnlineSubsystemOculus& InSubsystem);
+PRAGMA_ENABLE_DEPRECATION_WARNINGS
 
 	/**
 	* Default destructor
diff --git a/Engine/Plugins/Online/OnlineSubsystemOculus/Source/Private/OnlineMessageMultiTaskOculus.cpp b/Engine/Plugins/Online/OnlineSubsystemOculus/Source/Private/OnlineMessageMultiTaskOculus.cpp
--- a/Engine/Plugins/Online/OnlineSubsystemOculus/Source/Private/OnlineMessageMultiTaskOculus.cpp
+++ b/Engine/Plugins/Online/OnlineSubsystemOculus/Source/Private/OnlineMessageMultiTaskOculus.cpp
@@ -7,6 +7,7 @@
 void FOnlineMessageMultiTaskOculus::AddNewRequest(ovrRequest RequestId)
 {
 	InProgressRequests.Add(RequestId);
+PRAGMA_DISABLE_DEPRECATION_WARNINGS
 	OculusSubsystem.AddRequestDelegate(
 		RequestId,
 		FOculusMessageOnCompleteDelegate::CreateLambda([this, RequestId](ovrMessageHandle Message, bool bIsError)
@@ -22,4 +23,5 @@ void FOnlineMessageMultiTaskOculus::AddNewRequest(ovrRequest RequestId)
 			Delegate.ExecuteIfBound();
 		}
 	}));
+PRAGMA_ENABLE_DEPRECATION_WARNINGS
 }
diff --git a/Engine/Plugins/Online/OnlineSubsystemOculus/Source/Private/OnlineMessageMultiTaskOculus.h b/Engine/Plugins/Online/OnlineSubsystemOculus/Source/Private/OnlineMessageMultiTaskOculus.h
--- a/Engine/Plugins/Online/OnlineSubsystemOculus/Source/Private/OnlineMessageMultiTaskOculus.h
+++ b/Engine/Plugins/Online/OnlineSubsystemOculus/Source/Private/OnlineMessageMultiTaskOculus.h
@@ -25,10 +25,14 @@ protected:
 	FOnlineMessageMultiTaskOculus::FFinalizeDelegate Delegate;
 
 PACKAGE_SCOPE:
+PRAGMA_DISABLE_DEPRECATION_WARNINGS
 	FOnlineSubsystemOculus& OculusSubsystem;
+PRAGMA_ENABLE_DEPRECATION_WARNINGS
 
 	FOnlineMessageMultiTaskOculus(
+PRAGMA_DISABLE_DEPRECATION_WARNINGS
 		FOnlineSubsystemOculus& InOculusSubsystem,
+PRAGMA_ENABLE_DEPRECATION_WARNINGS
 		const FOnlineMessageMultiTaskOculus::FFinalizeDelegate& InDelegate)
 		: Delegate(InDelegate)
 		, OculusSubsystem(InOculusSubsystem)
diff --git a/Engine/Plugins/Online/OnlineSubsystemOculus/Source/Private/OnlineSessionInterfaceOculus.cpp b/Engine/Plugins/Online/OnlineSubsystemOculus/Source/Private/OnlineSessionInterfaceOculus.cpp
--- a/Engine/Plugins/Online/OnlineSubsystemOculus/Source/Private/OnlineSessionInterfaceOculus.cpp
+++ b/Engine/Plugins/Online/OnlineSubsystemOculus/Source/Private/OnlineSessionInterfaceOculus.cpp
@@ -15,7 +15,7 @@ FOnlineSessionInfoOculus::FOnlineSessionInfoOculus(ovrID RoomId) :
 /**
 * FOnlineSessionOculus
 */
-
+PRAGMA_DISABLE_DEPRECATION_WARNINGS
 FOnlineSessionOculus::FOnlineSessionOculus(FOnlineSubsystemOculus& InSubsystem) :
 	OculusSubsystem(InSubsystem)
 {
@@ -31,9 +31,11 @@ FOnlineSessionOculus::FOnlineSessionOculus(FOnlineSubsystemOculus& InSubsystem)
 		OculusSubsystem.GetNotifDelegate(ovrMessage_Notification_Matchmaking_MatchFound)
 		.AddRaw(this, &FOnlineSessionOculus::OnMatchmakingNotificationMatchFound);
 }
+PRAGMA_ENABLE_DEPRECATION_WARNINGS
 
 FOnlineSessionOculus::~FOnlineSessionOculus()
 {
+PRAGMA_DISABLE_DEPRECATION_WARNINGS
 	if (OnRoomNotificationUpdateHandle.IsValid())
 	{
 		OculusSubsystem.RemoveNotifDelegate(ovrMessage_Notification_Room_RoomUpdate, OnRoomNotificationUpdateHandle);
@@ -51,7 +53,7 @@ FOnlineSessionOculus::~FOnlineSessionOculus()
 		OculusSubsystem.RemoveNotifDelegate(ovrMessage_Notification_Matchmaking_MatchFound, OnMatchmakingNotificationMatchFoundHandle);
 		OnMatchmakingNotificationMatchFoundHandle.Reset();
 	}
-
+PRAGMA_ENABLE_DEPRECATION_WARNINGS
 	PendingInviteAcceptedSessions.Empty();
 
 	// Make sure the player leaves all the sessions they were in before destroying this
@@ -100,7 +102,9 @@ bool FOnlineSessionOculus::CreateSession(int32 HostingPlayerNum, FName SessionNa
 		return false;
 	}
 
+PRAGMA_DISABLE_DEPRECATION_WARNINGS
 	IOnlineIdentityPtr Identity = OculusSubsystem.GetIdentityInterface();
+PRAGMA_ENABLE_DEPRECATION_WARNINGS
 	if (!Identity.IsValid())
 	{
 		UE_LOG_ONLINE_SESSION(Warning, TEXT("No valid oculus identity interface."));
@@ -122,7 +126,9 @@ bool FOnlineSessionOculus::CreateSession(int32 HostingPlayerNum, FName SessionNa
 	Session->NumOpenPublicConnections = NewSessionSettings.NumPublicConnections;
 
 	Session->HostingPlayerNum = HostingPlayerNum;
+PRAGMA_DISABLE_DEPRECATION_WARNINGS
 	Session->LocalOwnerId = OculusSubsystem.GetIdentityInterface()->GetUniquePlayerId(HostingPlayerNum);
+PRAGMA_ENABLE_DEPRECATION_WARNINGS
 
 	// Setup the join policy
 	ovrRoomJoinPolicy JoinPolicy = ovrRoom_JoinPolicyInvitedUsers;
@@ -193,10 +199,11 @@ bool FOnlineSessionOculus::CreateMatchmakingSession(FNamedOnlineSession& Session
 	auto RequestId = (Session.SessionSettings.bShouldAdvertise)
 		? ovr_Matchmaking_CreateAndEnqueueRoom2(TCHAR_TO_ANSI(*Pool), MatchmakingOptions)
 		: ovr_Matchmaking_CreateRoom2(TCHAR_TO_ANSI(*Pool), MatchmakingOptions);
+PRAGMA_DISABLE_DEPRECATION_WARNINGS
 	OculusSubsystem.AddRequestDelegate(
 		RequestId,
 		FOculusMessageOnCompleteDelegate::CreateRaw(this, &FOnlineSessionOculus::OnCreateRoomComplete, Session.SessionName));
-
+PRAGMA_ENABLE_DEPRECATION_WARNINGS
 	ovr_MatchmakingOptions_Destroy(MatchmakingOptions);
 
 	return true;
@@ -225,12 +232,12 @@ bool FOnlineSessionOculus::CreateRoomSession(FNamedOnlineSession& Session, ovrRo
 	);
 
 	unsigned int MaxUsers = Session.SessionSettings.NumPublicConnections + Session.SessionSettings.NumPrivateConnections;
-
+PRAGMA_DISABLE_DEPRECATION_WARNINGS
 	OculusSubsystem.AddRequestDelegate(
 		ovr_Room_CreateAndJoinPrivate2(JoinPolicy, MaxUsers, RoomOptions),
 		FOculusMessageOnCompleteDelegate::CreateRaw(this, &FOnlineSessionOculus::OnCreateRoomComplete, Session.SessionName));
 	ovr_RoomOptions_Destroy(RoomOptions);
-
+PRAGMA_ENABLE_DEPRECATION_WARNINGS
 	return true;
 }
 
@@ -327,7 +334,9 @@ bool FOnlineSessionOculus::UpdateSession(FName SessionName, FOnlineSessionSettin
 		return false;
 	}
 
+PRAGMA_DISABLE_DEPRECATION_WARNINGS
 	auto LoggedInPlayerId = OculusSubsystem.GetIdentityInterface()->GetUniquePlayerId(0);
+PRAGMA_ENABLE_DEPRECATION_WARNINGS
 	if (!LoggedInPlayerId.IsValid() || *Session->OwningUserId != *LoggedInPlayerId)
 	{
 		UE_LOG_ONLINE_SESSION(Warning, TEXT("Need to own session (%s) before updating.  Current Owner: %s"), *SessionName.ToString(), *Session->OwningUserName);
@@ -368,7 +377,7 @@ bool FOnlineSessionOculus::UpdateMatchmakingRoom(FName SessionName, FOnlineSessi
 		auto RequestId = (UpdatedSessionSettings.bShouldAdvertise)
 			? ovr_Matchmaking_EnqueueRoom(GetOvrIDFromSession(*Session), nullptr)
 			: ovr_Matchmaking_Cancel2();
-
+PRAGMA_DISABLE_DEPRECATION_WARNINGS
 		OculusSubsystem.AddRequestDelegate(
 			RequestId,
 			FOculusMessageOnCompleteDelegate::CreateLambda([this, SessionName, UpdatedSettingsRef](ovrMessageHandle Message, bool bIsError)
@@ -398,6 +407,7 @@ bool FOnlineSessionOculus::UpdateMatchmakingRoom(FName SessionName, FOnlineSessi
 
 			UpdateRoomDataStore(SessionName, *SessionSettings);
 		}));
+PRAGMA_ENABLE_DEPRECATION_WARNINGS
 	}
 	else
 	{
@@ -467,6 +477,7 @@ bool FOnlineSessionOculus::UpdateRoomDataStore(FName SessionName, FOnlineSession
 	// If there is a delta, then fire off the request
 	if (DataStore.Num() > 0)
 	{
+PRAGMA_DISABLE_DEPRECATION_WARNINGS
 		OculusSubsystem.AddRequestDelegate(
 			ovr_Room_UpdateDataStore(GetOvrIDFromSession(*Session), DataStore.GetData(), DataStore.Num()),
 			FOculusMessageOnCompleteDelegate::CreateLambda([this, SessionName](ovrMessageHandle Message, bool bIsError)
@@ -494,6 +505,7 @@ bool FOnlineSessionOculus::UpdateRoomDataStore(FName SessionName, FOnlineSession
 
 			TriggerOnUpdateSessionCompleteDelegates(SessionName, true);
 		}));
+PRAGMA_ENABLE_DEPRECATION_WARNINGS
 	}
 	else
 	{
@@ -541,7 +553,7 @@ bool FOnlineSessionOculus::DestroySession(FName SessionName, const FOnDestroySes
 	auto RoomId = GetOvrIDFromSession(*Session);
 
 	Session->SessionState = EOnlineSessionState::Destroying;
-
+PRAGMA_DISABLE_DEPRECATION_WARNINGS
 	OculusSubsystem.AddRequestDelegate(
 		ovr_Room_Leave(RoomId),
 		FOculusMessageOnCompleteDelegate::CreateLambda([this, SessionName, Session, CompletionDelegate](ovrMessageHandle Message, bool bIsError)
@@ -561,7 +573,7 @@ bool FOnlineSessionOculus::DestroySession(FName SessionName, const FOnDestroySes
 		CompletionDelegate.ExecuteIfBound(SessionName, true);
 		TriggerOnDestroySessionCompleteDelegates(SessionName, true);
 	}));
-
+PRAGMA_ENABLE_DEPRECATION_WARNINGS
 	return true;
 }
 
@@ -618,7 +630,7 @@ bool FOnlineSessionOculus::StartMatchmaking(const TArray< FUniqueNetIdRef >& Loc
 	SearchSettings->SearchState = EOnlineAsyncTaskState::InProgress;
 	InProgressMatchmakingSearch = SearchSettings;
 	InProgressMatchmakingSearchName = SessionName;
-
+PRAGMA_DISABLE_DEPRECATION_WARNINGS
 	OculusSubsystem.AddRequestDelegate(
 		ovr_Matchmaking_Enqueue2(TCHAR_TO_UTF8(*Pool), nullptr),
 		FOculusMessageOnCompleteDelegate::CreateLambda([this, SessionName, SearchSettings](ovrMessageHandle Message, bool bIsError)
@@ -635,12 +647,13 @@ bool FOnlineSessionOculus::StartMatchmaking(const TArray< FUniqueNetIdRef >& Loc
 		// Nothing to trigger here.  
 		// If a match is found, TriggerOnMatchmakingCompleteDelegates() from the notification
 	}));
-
+PRAGMA_ENABLE_DEPRECATION_WARNINGS
 	return true;
 }
 
 bool FOnlineSessionOculus::CancelMatchmaking(int32 SearchingPlayerNum, FName SessionName)
 {
+PRAGMA_DISABLE_DEPRECATION_WARNINGS
 	OculusSubsystem.AddRequestDelegate(
 		ovr_Matchmaking_Cancel2(),
 		FOculusMessageOnCompleteDelegate::CreateLambda([this, SessionName](ovrMessageHandle Message, bool bIsError)
@@ -659,7 +672,7 @@ bool FOnlineSessionOculus::CancelMatchmaking(int32 SearchingPlayerNum, FName Ses
 		}
 		TriggerOnCancelMatchmakingCompleteDelegates(SessionName, true);
 	}));
-
+PRAGMA_ENABLE_DEPRECATION_WARNINGS
 	return true;
 }
 
@@ -705,6 +718,7 @@ bool FOnlineSessionOculus::FindSessions(const FUniqueNetId& SearchingPlayerId, c
 bool FOnlineSessionOculus::FindModeratedRoomSessions(const TSharedRef<FOnlineSessionSearch>& SearchSettings)
 {
 	SearchSettings->SearchState = EOnlineAsyncTaskState::InProgress;
+PRAGMA_DISABLE_DEPRECATION_WARNINGS
 	OculusSubsystem.AddRequestDelegate(
 		ovr_Room_GetModeratedRooms(),
 		FOculusMessageOnCompleteDelegate::CreateLambda([this, SearchSettings](ovrMessageHandle Message, bool bIsError)
@@ -760,7 +774,7 @@ bool FOnlineSessionOculus::FindModeratedRoomSessions(const TSharedRef<FOnlineSes
 		SearchSettings->SearchState = EOnlineAsyncTaskState::Done;
 		TriggerOnFindSessionsCompleteDelegates(true);
 	}));
-
+PRAGMA_ENABLE_DEPRECATION_WARNINGS
 	return true;
 }
 
@@ -774,6 +788,7 @@ bool FOnlineSessionOculus::FindMatchmakingSessions(const FString Pool, const TSh
 	SearchSettings->SearchState = EOnlineAsyncTaskState::InProgress;
 	InProgressMatchmakingSearch = SearchSettings;
 
+PRAGMA_DISABLE_DEPRECATION_WARNINGS
 	OculusSubsystem.AddRequestDelegate(
 		ovr_Matchmaking_Browse2(TCHAR_TO_UTF8(*Pool), nullptr),
 		FOculusMessageOnCompleteDelegate::CreateLambda([this, SearchSettings](ovrMessageHandle Message, bool bIsError)
@@ -827,13 +842,15 @@ bool FOnlineSessionOculus::FindMatchmakingSessions(const FString Pool, const TSh
 		SearchSettings->SearchState = EOnlineAsyncTaskState::Done;
 		TriggerOnFindSessionsCompleteDelegates(true);
 	}));
-
+PRAGMA_ENABLE_DEPRECATION_WARNINGS
 	return true;
 }
 
 bool FOnlineSessionOculus::FindSessionById(const FUniqueNetId& SearchingUserId, const FUniqueNetId& SessionId, const FUniqueNetId& FriendId, const FOnSingleSessionResultCompleteDelegate& CompletionDelegate)
 {
+PRAGMA_DISABLE_DEPRECATION_WARNINGS
 	auto LoggedInPlayerId = OculusSubsystem.GetIdentityInterface()->GetUniquePlayerId(0);
+PRAGMA_ENABLE_DEPRECATION_WARNINGS
 	if (!LoggedInPlayerId.IsValid() || SearchingUserId != *LoggedInPlayerId)
 	{
 		UE_LOG_ONLINE_SESSION(Warning, TEXT("Can only search session with logged in player"));
@@ -847,6 +864,7 @@ bool FOnlineSessionOculus::FindSessionById(const FUniqueNetId& SearchingUserId,
 	}
 
 	auto RoomId = static_cast<const FUniqueNetIdOculus&>(SessionId);
+PRAGMA_DISABLE_DEPRECATION_WARNINGS
 	OculusSubsystem.AddRequestDelegate(
 		ovr_Room_Get(RoomId.GetID()),
 		FOculusMessageOnCompleteDelegate::CreateLambda([this, CompletionDelegate](ovrMessageHandle Message, bool bIsError)
@@ -887,7 +905,7 @@ bool FOnlineSessionOculus::FindSessionById(const FUniqueNetId& SearchingUserId,
 		auto RoomJoinability = ovr_Room_GetJoinability(Room);
 		CompletionDelegate.ExecuteIfBound(0, RoomJoinability == ovrRoom_JoinabilityCanJoin, SearchResult);
 	}));
-	
+PRAGMA_ENABLE_DEPRECATION_WARNINGS	
 	return true;
 }
 
@@ -922,11 +940,13 @@ bool FOnlineSessionOculus::JoinSession(int32 PlayerNum, FName SessionName, const
 	check(Session);
 	Session->SessionState = EOnlineSessionState::Creating;
 	Session->HostingPlayerNum = PlayerNum;
+PRAGMA_DISABLE_DEPRECATION_WARNINGS
 	Session->LocalOwnerId = OculusSubsystem.GetIdentityInterface()->GetUniquePlayerId(PlayerNum);
-
+PRAGMA_ENABLE_DEPRECATION_WARNINGS
 	auto SearchSessionInfo = static_cast<const FOnlineSessionInfoOculus*>(DesiredSession.Session.SessionInfo.Get());
 	auto RoomId = FUniqueNetIdOculus::Cast(SearchSessionInfo->GetSessionId()).GetID();
 
+PRAGMA_DISABLE_DEPRECATION_WARNINGS
 	OculusSubsystem.AddRequestDelegate(
 		ovr_Room_Join(RoomId, /* subscribeToUpdates */ true),
 		FOculusMessageOnCompleteDelegate::CreateLambda([this, SessionName, Session](ovrMessageHandle Message, bool bIsError)
@@ -958,7 +978,7 @@ bool FOnlineSessionOculus::JoinSession(int32 PlayerNum, FName SessionName, const
 
 		TriggerOnJoinSessionCompleteDelegates(SessionName, EOnJoinSessionCompleteResult::Success);
 	}));
-
+PRAGMA_ENABLE_DEPRECATION_WARNINGS
 	return true;
 }
 
@@ -970,7 +990,7 @@ bool FOnlineSessionOculus::JoinSession(const FUniqueNetId& PlayerId, FName Sessi
 bool FOnlineSessionOculus::FindFriendSession(int32 LocalUserNum, const FUniqueNetId& Friend)
 {
 	auto OculusId = static_cast<const FUniqueNetIdOculus&>(Friend);
-
+PRAGMA_DISABLE_DEPRECATION_WARNINGS
 	OculusSubsystem.AddRequestDelegate(
 		ovr_Room_GetCurrentForUser(OculusId.GetID()),
 		FOculusMessageOnCompleteDelegate::CreateLambda([this, LocalUserNum](ovrMessageHandle Message, bool bIsError)
@@ -1013,7 +1033,7 @@ bool FOnlineSessionOculus::FindFriendSession(int32 LocalUserNum, const FUniqueNe
 		auto RoomJoinability = ovr_Room_GetJoinability(Room);
 		TriggerOnFindFriendSessionCompleteDelegates(LocalUserNum, RoomJoinability == ovrRoom_JoinabilityCanJoin, SearchResult);
 	}));
-
+PRAGMA_ENABLE_DEPRECATION_WARNINGS
 	return true;
 };
 
@@ -1052,8 +1072,9 @@ bool FOnlineSessionOculus::SendSessionInviteToFriends(int32 LocalUserNum, FName
 		UE_LOG_ONLINE_SESSION(Warning, TEXT("Session (%s) doesn't exist"), *SessionName.ToString());
 		return false;
 	}
-
+PRAGMA_DISABLE_DEPRECATION_WARNINGS
 	IOnlineFriendsPtr FriendsInterface = OculusSubsystem.GetFriendsInterface();
+PRAGMA_ENABLE_DEPRECATION_WARNINGS
 	if (!FriendsInterface.IsValid())
 	{
 		UE_LOG_ONLINE_SESSION(Warning, TEXT("Cannot get invite tokens for friends"));
@@ -1283,9 +1304,10 @@ void FOnlineSessionOculus::OnRoomNotificationUpdate(ovrMessageHandle Message, bo
 
 void FOnlineSessionOculus::OnRoomInviteAccepted(ovrMessageHandle Message, bool bIsError)
 {
+PRAGMA_DISABLE_DEPRECATION_WARNINGS
 	IOnlineIdentityPtr Identity = OculusSubsystem.GetIdentityInterface();
 	auto PlayerId = Identity->GetUniquePlayerId(0);
-
+PRAGMA_ENABLE_DEPRECATION_WARNINGS
 	FOnlineSessionSearchResult SearchResult;
 	if (bIsError)
 	{
@@ -1308,7 +1330,7 @@ void FOnlineSessionOculus::OnRoomInviteAccepted(ovrMessageHandle Message, bool b
 	// Fetch the room details to create the SessionResult
 
 	ovr_Room_Get(RoomId);
-
+PRAGMA_DISABLE_DEPRECATION_WARNINGS
 	OculusSubsystem.AddRequestDelegate(
 		ovr_Room_Get(RoomId),
 		FOculusMessageOnCompleteDelegate::CreateLambda([this, PlayerId](ovrMessageHandle InMessage, bool bInIsError)
@@ -1338,6 +1360,7 @@ void FOnlineSessionOculus::OnRoomInviteAccepted(ovrMessageHandle Message, bool b
 		LocalSearchResult.Session = Session.Get();
 		TriggerOnSessionUserInviteAcceptedDelegates(true, 0, PlayerId, LocalSearchResult);
 	}));
+PRAGMA_ENABLE_DEPRECATION_WARNINGS
 }
 
 void FOnlineSessionOculus::OnMatchmakingNotificationMatchFound(ovrMessageHandle Message, bool bIsError)
@@ -1477,7 +1500,9 @@ void FOnlineSessionOculus::TickPendingInvites(float DeltaTime)
 {
 	if (PendingInviteAcceptedSessions.Num() > 0 && OnSessionUserInviteAcceptedDelegates.IsBound())
 	{
+PRAGMA_DISABLE_DEPRECATION_WARNINGS
 		IOnlineIdentityPtr Identity = OculusSubsystem.GetIdentityInterface();
+PRAGMA_ENABLE_DEPRECATION_WARNINGS
 		auto PlayerId = Identity->GetUniquePlayerId(0);
 		for (auto PendingInviteAcceptedSession : PendingInviteAcceptedSessions)
 		{
diff --git a/Engine/Plugins/Online/OnlineSubsystemOculus/Source/Private/OnlineSessionInterfaceOculus.h b/Engine/Plugins/Online/OnlineSubsystemOculus/Source/Private/OnlineSessionInterfaceOculus.h
--- a/Engine/Plugins/Online/OnlineSubsystemOculus/Source/Private/OnlineSessionInterfaceOculus.h
+++ b/Engine/Plugins/Online/OnlineSubsystemOculus/Source/Private/OnlineSessionInterfaceOculus.h
@@ -22,8 +22,10 @@ class FOnlineSessionOculus : public IOnlineSession
 {
 private:
 
+PRAGMA_DISABLE_DEPRECATION_WARNINGS
 	/** Reference to the main Oculus subsystem */
 	FOnlineSubsystemOculus& OculusSubsystem;
+PRAGMA_ENABLE_DEPRECATION_WARNINGS
 
 	/**
 	 * Current session settings
@@ -75,8 +77,9 @@ PACKAGE_SCOPE:
 	bool UpdateRoomDataStore(FName SessionName, FOnlineSessionSettings& UpdatedSessionSettings);
 
 public:
-
+PRAGMA_DISABLE_DEPRECATION_WARNINGS
 	FOnlineSessionOculus(FOnlineSubsystemOculus& InSubsystem);
+PRAGMA_ENABLE_DEPRECATION_WARNINGS
 	virtual ~FOnlineSessionOculus();
 
 	// IOnlineSession
diff --git a/Engine/Plugins/Online/OnlineSubsystemOculus/Source/Private/OnlineSubsystemOculus.cpp b/Engine/Plugins/Online/OnlineSubsystemOculus/Source/Private/OnlineSubsystemOculus.cpp
--- a/Engine/Plugins/Online/OnlineSubsystemOculus/Source/Private/OnlineSubsystemOculus.cpp
+++ b/Engine/Plugins/Online/OnlineSubsystemOculus/Source/Private/OnlineSubsystemOculus.cpp
@@ -180,6 +180,7 @@ bool FOnlineSubsystemOculus::Tick(float DeltaTime)
 	return true;
 }
 
+PRAGMA_DISABLE_DEPRECATION_WARNINGS
 void FOnlineSubsystemOculus::AddRequestDelegate(ovrRequest RequestId, FOculusMessageOnCompleteDelegate&& Delegate) const
 {
 	check(MessageTaskManager);
@@ -197,6 +198,7 @@ void FOnlineSubsystemOculus::RemoveNotifDelegate(ovrMessageType MessageType, con
 	check(MessageTaskManager);
 	return MessageTaskManager->RemoveNotifDelegate(MessageType, Delegate);
 }
+PRAGMA_ENABLE_DEPRECATION_WARNINGS
 
 bool FOnlineSubsystemOculus::Init()
 {
diff --git a/Engine/Plugins/Online/OnlineSubsystemOculus/Source/Private/OnlineSubsystemOculusModule.cpp b/Engine/Plugins/Online/OnlineSubsystemOculus/Source/Private/OnlineSubsystemOculusModule.cpp
--- a/Engine/Plugins/Online/OnlineSubsystemOculus/Source/Private/OnlineSubsystemOculusModule.cpp
+++ b/Engine/Plugins/Online/OnlineSubsystemOculus/Source/Private/OnlineSubsystemOculusModule.cpp
@@ -24,6 +24,9 @@ public:
 
 	IOnlineSubsystemPtr CreateSubsystem(FName InstanceName) override
 	{
+		UE_LOG_ONLINE_ONCE(Warning, TEXT("OnlineSubsystemOculus has been deprecated"));
+
+PRAGMA_DISABLE_DEPRECATION_WARNINGS
 		if (!OnlineSub.IsValid())
 		{
 			OnlineSub = MakeShared<FOnlineSubsystemOculus, ESPMode::ThreadSafe>(InstanceName);
@@ -50,7 +53,7 @@ public:
 			OnlineSub->Shutdown();
 			OnlineSub = nullptr;
 		}
-
+PRAGMA_ENABLE_DEPRECATION_WARNINGS
 		return OnlineSub;
 	}
 };
@@ -63,7 +66,9 @@ void FOnlineSubsystemOculusModule::StartupModule()
 
 	// Create and register our singleton factory with the main online subsystem for easy access
 	FOnlineSubsystemModule& OSS = FModuleManager::GetModuleChecked<FOnlineSubsystemModule>("OnlineSubsystem");
+PRAGMA_DISABLE_DEPRECATION_WARNINGS
 	OSS.RegisterPlatformService(OCULUS_SUBSYSTEM, OculusFactory);
+PRAGMA_ENABLE_DEPRECATION_WARNINGS
 }
 
 void FOnlineSubsystemOculusModule::ShutdownModule()
@@ -71,7 +76,9 @@ void FOnlineSubsystemOculusModule::ShutdownModule()
 	UE_LOG_ONLINE(Log, TEXT("Oculus Shutdown!"));
 
 	FOnlineSubsystemModule& OSS = FModuleManager::GetModuleChecked<FOnlineSubsystemModule>("OnlineSubsystem");
+PRAGMA_DISABLE_DEPRECATION_WARNINGS
 	OSS.UnregisterPlatformService(OCULUS_SUBSYSTEM);
+PRAGMA_ENABLE_DEPRECATION_WARNINGS
 	
 	delete OculusFactory;
 	OculusFactory = nullptr;
diff --git a/Engine/Plugins/Online/OnlineSubsystemOculus/Source/Private/OnlineSubsystemOculusTypes.h b/Engine/Plugins/Online/OnlineSubsystemOculus/Source/Private/OnlineSubsystemOculusTypes.h
--- a/Engine/Plugins/Online/OnlineSubsystemOculus/Source/Private/OnlineSubsystemOculusTypes.h
+++ b/Engine/Plugins/Online/OnlineSubsystemOculus/Source/Private/OnlineSubsystemOculusTypes.h
@@ -39,8 +39,10 @@ public:
 
 	static const FUniqueNetIdOculus& Cast(const FUniqueNetId& NetId)
 	{
+PRAGMA_DISABLE_DEPRECATION_WARNINGS
 		check(NetId.GetType() == OCULUS_SUBSYSTEM);
 		return *static_cast<const FUniqueNetIdOculus*>(&NetId);
+PRAGMA_ENABLE_DEPRECATION_WARNINGS
 	}
 
 	FUniqueNetIdOculusRef AsShared() const
@@ -50,7 +52,9 @@ public:
 
 	virtual FName GetType() const override
 	{
+PRAGMA_DISABLE_DEPRECATION_WARNINGS
 		return OCULUS_SUBSYSTEM;
+PRAGMA_ENABLE_DEPRECATION_WARNINGS
 	}
 
 	// IOnlinePlatformData
@@ -123,10 +127,12 @@ private:
 	*/
 	explicit FUniqueNetIdOculus(const FUniqueNetId& Src)
 	{
+PRAGMA_DISABLE_DEPRECATION_WARNINGS
 		if (Src.GetType() == OCULUS_SUBSYSTEM)
 		{
 			ID = FUniqueNetIdOculus::Cast(Src).ID;
 		}
+PRAGMA_ENABLE_DEPRECATION_WARNINGS
 	}
 };
 
diff --git a/Engine/Plugins/Online/OnlineSubsystemOculus/Source/Private/OnlineUserCloudOculus.cpp b/Engine/Plugins/Online/OnlineSubsystemOculus/Source/Private/OnlineUserCloudOculus.cpp
--- a/Engine/Plugins/Online/OnlineSubsystemOculus/Source/Private/OnlineUserCloudOculus.cpp
+++ b/Engine/Plugins/Online/OnlineSubsystemOculus/Source/Private/OnlineUserCloudOculus.cpp
@@ -20,6 +20,7 @@ const FString FOnlineUserCloudOculus::DEFAULT_BUCKET_KEY = TEXT("DefaultUserClou
 
 const FString FOnlineUserCloudOculus::ALL_BUCKETS_KEY = TEXT("UserCloudBuckets");
 
+PRAGMA_DISABLE_DEPRECATION_WARNINGS
 FOnlineUserCloudOculus::FOnlineUserCloudOculus(FOnlineSubsystemOculus& InSubsystem)
 	: OculusSubsystem(InSubsystem)
 	, EnumerateBucketsCounter(-1)
@@ -36,10 +37,13 @@ FOnlineUserCloudOculus::FOnlineUserCloudOculus(FOnlineSubsystemOculus& InSubsyst
 		Buckets.AddUnique(DefaultBucket);
 	}
 }
+PRAGMA_ENABLE_DEPRECATION_WARNINGS
 
 void FOnlineUserCloudOculus::EnumerateUserFiles(const FUniqueNetId& UserId)
 {
+PRAGMA_DISABLE_DEPRECATION_WARNINGS
 	auto LoggedInPlayerId = OculusSubsystem.GetIdentityInterface()->GetUniquePlayerId(0);
+PRAGMA_ENABLE_DEPRECATION_WARNINGS
 	if (!LoggedInPlayerId.IsValid() || UserId != *LoggedInPlayerId)
 	{
 		UE_LOG_ONLINE_CLOUD(Warning, TEXT("Can only enumerate save data for logged in player"));
@@ -112,7 +116,7 @@ void FOnlineUserCloudOculus::RequestEnumeratePagedBuckets(FUniqueNetIdPtr UserId
 			TriggerOnEnumerateUserFilesCompleteDelegates(true, *UserId);
 		}
 	});
-
+PRAGMA_DISABLE_DEPRECATION_WARNINGS
 	if (previousPage)
 	{
 		OculusSubsystem.AddRequestDelegate(ovr_CloudStorage_GetNextCloudStorageMetadataArrayPage(previousPage), std::move(Delegate));
@@ -121,11 +125,14 @@ void FOnlineUserCloudOculus::RequestEnumeratePagedBuckets(FUniqueNetIdPtr UserId
 	{
 		OculusSubsystem.AddRequestDelegate(ovr_CloudStorage_LoadBucketMetadata(TCHAR_TO_UTF8(*Buckets[EnumerateBucketsCounter])), std::move(Delegate));
 	}
+PRAGMA_ENABLE_DEPRECATION_WARNINGS
 }
 
 void FOnlineUserCloudOculus::GetUserFileList(const FUniqueNetId& UserId, TArray<FCloudFileHeader>& UserFiles)
 {
+PRAGMA_DISABLE_DEPRECATION_WARNINGS
 	auto LoggedInPlayerId = OculusSubsystem.GetIdentityInterface()->GetUniquePlayerId(0);
+PRAGMA_ENABLE_DEPRECATION_WARNINGS
 	if (!LoggedInPlayerId.IsValid() || UserId != *LoggedInPlayerId)
 	{
 		UE_LOG_ONLINE_CLOUD(Warning, TEXT("Can get file list for logged in player"));
@@ -137,7 +144,9 @@ void FOnlineUserCloudOculus::GetUserFileList(const FUniqueNetId& UserId, TArray<
 
 bool FOnlineUserCloudOculus::WriteUserFile(const FUniqueNetId& UserId, const FString& FileName, TArray<uint8>& FileContents, bool bCompressBeforeUpload)
 {
+PRAGMA_DISABLE_DEPRECATION_WARNINGS
 	auto LoggedInPlayerId = OculusSubsystem.GetIdentityInterface()->GetUniquePlayerId(0);
+PRAGMA_ENABLE_DEPRECATION_WARNINGS
 	if (!LoggedInPlayerId.IsValid() || UserId != *LoggedInPlayerId)
 	{
 		UE_LOG_ONLINE_CLOUD(Warning, TEXT("Can only save data for logged in player"));
@@ -172,17 +181,19 @@ bool FOnlineUserCloudOculus::WriteUserFile(const FUniqueNetId& UserId, const FSt
 		delete TmpBuffer;
 		TriggerOnWriteUserFileCompleteDelegates(!bIsError, *LoggedInPlayerId, FileName);
 	});
-
+PRAGMA_DISABLE_DEPRECATION_WARNINGS
 	OculusSubsystem.AddRequestDelegate(
 		ovr_CloudStorage_Save(TCHAR_TO_UTF8(*BucketName), TCHAR_TO_UTF8(*Key), TmpBuffer->GetData(), TmpBuffer->Num(), 0, nullptr),
 		std::move(DelegateLambda));
-
+PRAGMA_ENABLE_DEPRECATION_WARNINGS
 	return true;
 }
 
 bool FOnlineUserCloudOculus::ReadUserFile(const FUniqueNetId& UserId, const FString& FileName)
 {
+PRAGMA_DISABLE_DEPRECATION_WARNINGS
 	auto LoggedInPlayerId = OculusSubsystem.GetIdentityInterface()->GetUniquePlayerId(0);
+PRAGMA_ENABLE_DEPRECATION_WARNINGS
 	if (!LoggedInPlayerId.IsValid() || UserId != *LoggedInPlayerId)
 	{
 		UE_LOG_ONLINE_CLOUD(Warning, TEXT("Can only read data for logged in player"));
@@ -197,6 +208,7 @@ bool FOnlineUserCloudOculus::ReadUserFile(const FUniqueNetId& UserId, const FStr
 		Key = FileName;
 	}
 
+PRAGMA_DISABLE_DEPRECATION_WARNINGS
 	OculusSubsystem.AddRequestDelegate(
 		ovr_CloudStorage_Load(TCHAR_TO_UTF8(*BucketName), TCHAR_TO_UTF8(*Key)),
 		FOculusMessageOnCompleteDelegate::CreateLambda([this, BucketName, Key, LoggedInPlayerId, FileName](ovrMessageHandle Message, bool bIsError)
@@ -221,13 +233,15 @@ bool FOnlineUserCloudOculus::ReadUserFile(const FUniqueNetId& UserId, const FStr
 		}
 		TriggerOnReadUserFileCompleteDelegates(!bIsError, *LoggedInPlayerId, FileName);
 	}));
-
+PRAGMA_ENABLE_DEPRECATION_WARNINGS
 	return true;
 }
 
 bool FOnlineUserCloudOculus::GetFileContents(const FUniqueNetId& UserId, const FString& FileName, TArray<uint8>& FileContents)
 {
+PRAGMA_DISABLE_DEPRECATION_WARNINGS
 	auto LoggedInPlayerId = OculusSubsystem.GetIdentityInterface()->GetUniquePlayerId(0);
+PRAGMA_ENABLE_DEPRECATION_WARNINGS
 	if (!LoggedInPlayerId.IsValid() || UserId != *LoggedInPlayerId)
 	{
 		UE_LOG_ONLINE_CLOUD(Warning, TEXT("Can only read data for logged in player"));
@@ -246,7 +260,9 @@ bool FOnlineUserCloudOculus::GetFileContents(const FUniqueNetId& UserId, const F
 
 bool FOnlineUserCloudOculus::ClearFiles(const FUniqueNetId& UserId)
 {
+PRAGMA_DISABLE_DEPRECATION_WARNINGS
 	auto LoggedInPlayerId = OculusSubsystem.GetIdentityInterface()->GetUniquePlayerId(0);
+PRAGMA_ENABLE_DEPRECATION_WARNINGS
 	if (!LoggedInPlayerId.IsValid() || UserId != *LoggedInPlayerId)
 	{
 		UE_LOG_ONLINE_CLOUD(Warning, TEXT("Can only clear data for logged in player"));
@@ -259,7 +275,9 @@ bool FOnlineUserCloudOculus::ClearFiles(const FUniqueNetId& UserId)
 
 bool FOnlineUserCloudOculus::ClearFile(const FUniqueNetId& UserId, const FString& FileName)
 {
+PRAGMA_DISABLE_DEPRECATION_WARNINGS
 	auto LoggedInPlayerId = OculusSubsystem.GetIdentityInterface()->GetUniquePlayerId(0);
+PRAGMA_ENABLE_DEPRECATION_WARNINGS
 	if (!LoggedInPlayerId.IsValid() || UserId != *LoggedInPlayerId)
 	{
 		UE_LOG_ONLINE_CLOUD(Warning, TEXT("Can only clear data for logged in player"));
@@ -282,7 +300,9 @@ void FOnlineUserCloudOculus::CancelWriteUserFile(const FUniqueNetId& UserId, con
 
 bool FOnlineUserCloudOculus::DeleteUserFile(const FUniqueNetId& UserId, const FString& FileName, bool bShouldCloudDelete, bool bShouldLocallyDelete)
 {
+PRAGMA_DISABLE_DEPRECATION_WARNINGS
 	auto LoggedInPlayerId = OculusSubsystem.GetIdentityInterface()->GetUniquePlayerId(0);
+PRAGMA_ENABLE_DEPRECATION_WARNINGS
 	if (!LoggedInPlayerId.IsValid() || UserId != *LoggedInPlayerId)
 	{
 		UE_LOG_ONLINE_CLOUD(Warning, TEXT("Can only save data for logged in player"));
@@ -313,10 +333,11 @@ bool FOnlineUserCloudOculus::DeleteUserFile(const FUniqueNetId& UserId, const FS
 
 			TriggerOnDeleteUserFileCompleteDelegates(!bIsError, *LoggedInPlayerId, FileName);
 		});
-
+PRAGMA_DISABLE_DEPRECATION_WARNINGS
 		OculusSubsystem.AddRequestDelegate(
 			ovr_CloudStorage_Delete(TCHAR_TO_UTF8(*BucketName), TCHAR_TO_UTF8(*Key)),
 			std::move(DelegateLambda));
+PRAGMA_ENABLE_DEPRECATION_WARNINGS
 	}
 	else
 	{
diff --git a/Engine/Plugins/Online/OnlineSubsystemOculus/Source/Private/OnlineUserCloudOculus.h b/Engine/Plugins/Online/OnlineSubsystemOculus/Source/Private/OnlineUserCloudOculus.h
--- a/Engine/Plugins/Online/OnlineSubsystemOculus/Source/Private/OnlineUserCloudOculus.h
+++ b/Engine/Plugins/Online/OnlineSubsystemOculus/Source/Private/OnlineUserCloudOculus.h
@@ -20,7 +20,10 @@
 class FOnlineUserCloudOculus : public IOnlineUserCloud
 {
 public:
+PRAGMA_DISABLE_DEPRECATION_WARNINGS
 	FOnlineUserCloudOculus(FOnlineSubsystemOculus& InSubsystem);
+PRAGMA_ENABLE_DEPRECATION_WARNINGS
+
 	virtual ~FOnlineUserCloudOculus() = default;
 
 	virtual void EnumerateUserFiles(const FUniqueNetId& UserId) override;
@@ -58,7 +61,9 @@ private:
 	/** Tag used to find the all bucket names in DefaultEngine.ini */
 	static const FString ALL_BUCKETS_KEY;
 
+PRAGMA_DISABLE_DEPRECATION_WARNINGS
 	FOnlineSubsystemOculus& OculusSubsystem;
+PRAGMA_ENABLE_DEPRECATION_WARNINGS
 
 	/** Default Bucket to store saves if none is specified from FileName - set in DefaultEngine.ini */
 	FString DefaultBucket;
diff --git a/Engine/Plugins/Online/OnlineSubsystemOculus/Source/Private/OnlineVoiceOculus.cpp b/Engine/Plugins/Online/OnlineSubsystemOculus/Source/Private/OnlineVoiceOculus.cpp
--- a/Engine/Plugins/Online/OnlineSubsystemOculus/Source/Private/OnlineVoiceOculus.cpp
+++ b/Engine/Plugins/Online/OnlineSubsystemOculus/Source/Private/OnlineVoiceOculus.cpp
@@ -22,14 +22,17 @@ FRemoteTalkerDataOculus::FRemoteTalkerDataOculus() :
 {
 }
 
+PRAGMA_DISABLE_DEPRECATION_WARNINGS
 FOnlineVoiceOculus::FOnlineVoiceOculus(FOnlineSubsystemOculus& InSubsystem) :
 	bIsLocalPlayerMuted(false),
 	OculusSubsystem(InSubsystem)
 {
 }
+PRAGMA_ENABLE_DEPRECATION_WARNINGS
 
 FOnlineVoiceOculus::~FOnlineVoiceOculus()
 {
+PRAGMA_DISABLE_DEPRECATION_WARNINGS
 	if (VoipConnectionRequestDelegateHandle.IsValid())
 	{
 		OculusSubsystem.RemoveNotifDelegate(ovrMessage_Notification_Voip_ConnectRequest, VoipConnectionRequestDelegateHandle);
@@ -41,6 +44,7 @@ FOnlineVoiceOculus::~FOnlineVoiceOculus()
 		OculusSubsystem.RemoveNotifDelegate(ovrMessage_Notification_Voip_StateChange, VoipStateChangeDelegateHandle);
 		VoipStateChangeDelegateHandle.Reset();
 	}
+PRAGMA_ENABLE_DEPRECATION_WARNINGS
 }
 
 bool FOnlineVoiceOculus::Init()
@@ -48,11 +52,13 @@ bool FOnlineVoiceOculus::Init()
 	bool bHasVoiceEnabled = false;
 	if (GConfig->GetBool(TEXT("OnlineSubsystem"), TEXT("bHasVoiceEnabled"), bHasVoiceEnabled, GEngineIni) && bHasVoiceEnabled)
 	{
+PRAGMA_DISABLE_DEPRECATION_WARNINGS
 		VoipConnectionRequestDelegateHandle = OculusSubsystem.GetNotifDelegate(ovrMessage_Notification_Voip_ConnectRequest)
 			.AddRaw(this, &FOnlineVoiceOculus::OnVoipConnectionRequest);
 
 		VoipStateChangeDelegateHandle = OculusSubsystem.GetNotifDelegate(ovrMessage_Notification_Voip_StateChange)
 			.AddRaw(this, &FOnlineVoiceOculus::OnVoipStateChange);
+PRAGMA_ENABLE_DEPRECATION_WARNINGS
 	}
 	else
 	{
diff --git a/Engine/Plugins/Online/OnlineSubsystemOculus/Source/Public/OnlineSubsystemOculus.h b/Engine/Plugins/Online/OnlineSubsystemOculus/Source/Public/OnlineSubsystemOculus.h
--- a/Engine/Plugins/Online/OnlineSubsystemOculus/Source/Public/OnlineSubsystemOculus.h
+++ b/Engine/Plugins/Online/OnlineSubsystemOculus/Source/Public/OnlineSubsystemOculus.h
@@ -23,7 +23,7 @@ typedef TUniquePtr<class FOnlineMessageTaskManagerOculus> FOnlineMessageTaskMana
 /**
 *	OnlineSubsystemOculus - Implementation of the online subsystem for Oculus services
 */
-class ONLINESUBSYSTEMOCULUS_API FOnlineSubsystemOculus :
+class ONLINESUBSYSTEMOCULUS_API UE_DEPRECATED(5.2, "OnlineSubsystemOculus has been deprecated and will be removed") FOnlineSubsystemOculus :
 	public FOnlineSubsystemImpl
 {
 
@@ -87,7 +87,9 @@ PACKAGE_SCOPE:
 	/** Only the factory makes instances */
 	FOnlineSubsystemOculus() = delete;
 	explicit FOnlineSubsystemOculus(FName InInstanceName) :
+PRAGMA_DISABLE_DEPRECATION_WARNINGS
 		FOnlineSubsystemImpl(OCULUS_SUBSYSTEM, InInstanceName),
+PRAGMA_ENABLE_DEPRECATION_WARNINGS
 		bOculusInit(false)
 	{}
 
@@ -133,4 +135,6 @@ namespace FNetworkProtocolTypes
 	ONLINESUBSYSTEMOCULUS_API extern const FLazyName Oculus;
 }
 
+PRAGMA_DISABLE_DEPRECATION_WARNINGS
 typedef TSharedPtr<FOnlineSubsystemOculus, ESPMode::ThreadSafe> FOnlineSubsystemOculusPtr;
+PRAGMA_ENABLE_DEPRECATION_WARNINGS
