Fix for deprecation of Oculus OSS compile issue on Android
#rnx

diff --git a/Engine/Plugins/Online/OnlineSubsystemOculus/Source/Private/OnlineAchievementsInterfaceOculus.cpp b/Engine/Plugins/Online/OnlineSubsystemOculus/Source/Private/OnlineAchievementsInterfaceOculus.cpp
--- a/Engine/Plugins/Online/OnlineSubsystemOculus/Source/Private/OnlineAchievementsInterfaceOculus.cpp
+++ b/Engine/Plugins/Online/OnlineSubsystemOculus/Source/Private/OnlineAchievementsInterfaceOculus.cpp
@@ -57,10 +57,12 @@ PRAGMA_ENABLE_DEPRECATION_WARNINGS
 };
 TSet< TSharedRef<FOnlineMessageMultiTaskOculusWriteAchievements> > FOnlineMessageMultiTaskOculusWriteAchievements::ActiveAchievementWriteTasks;
 
+PRAGMA_DISABLE_DEPRECATION_WARNINGS
 FOnlineAchievementsOculus::FOnlineAchievementsOculus(class FOnlineSubsystemOculus& InSubsystem)
 : OculusSubsystem(InSubsystem)
 {
 }
+PRAGMA_ENABLE_DEPRECATION_WARNINGS
 
 FOnlineAchievementsOculus::~FOnlineAchievementsOculus()
 {
diff --git a/Engine/Plugins/Online/OnlineSubsystemOculus/Source/Private/OnlineLeaderboardInterfaceOculus.cpp b/Engine/Plugins/Online/OnlineSubsystemOculus/Source/Private/OnlineLeaderboardInterfaceOculus.cpp
--- a/Engine/Plugins/Online/OnlineSubsystemOculus/Source/Private/OnlineLeaderboardInterfaceOculus.cpp
+++ b/Engine/Plugins/Online/OnlineSubsystemOculus/Source/Private/OnlineLeaderboardInterfaceOculus.cpp
@@ -5,10 +5,12 @@
 #include "OnlineIdentityOculus.h"
 #include "OnlineSubsystemOculusPackage.h"
 
+PRAGMA_DISABLE_DEPRECATION_WARNINGS
 FOnlineLeaderboardOculus::FOnlineLeaderboardOculus(class FOnlineSubsystemOculus& InSubsystem)
 : OculusSubsystem(InSubsystem)
 {
 }
+PRAGMA_ENABLE_DEPRECATION_WARNINGS
 
 bool FOnlineLeaderboardOculus::ReadLeaderboards(const TArray< FUniqueNetIdRef >& Players, FOnlineLeaderboardReadRef& ReadObject)
 {
diff --git a/Engine/Plugins/Online/OnlineSubsystemOculus/Source/Private/OnlineVoiceOculus.h b/Engine/Plugins/Online/OnlineSubsystemOculus/Source/Private/OnlineVoiceOculus.h
--- a/Engine/Plugins/Online/OnlineSubsystemOculus/Source/Private/OnlineVoiceOculus.h
+++ b/Engine/Plugins/Online/OnlineSubsystemOculus/Source/Private/OnlineVoiceOculus.h
@@ -35,7 +35,9 @@ public:
 	*
 	* @param InSubsystem online subsystem being used
 	*/
+	PRAGMA_DISABLE_DEPRECATION_WARNINGS
 	FOnlineVoiceOculus(class FOnlineSubsystemOculus& InSubsystem);
+	PRAGMA_ENABLE_DEPRECATION_WARNINGS
 	virtual ~FOnlineVoiceOculus();
 
 	// IOnlineVoice
@@ -145,8 +147,10 @@ private:
 
 	int32 IndexOfRemoteTalker(const FUniqueNetId& TalkerId);
 
+	PRAGMA_DISABLE_DEPRECATION_WARNINGS
 	/** Reference to the main Oculus subsystem */
 	class FOnlineSubsystemOculus& OculusSubsystem;
+	PRAGMA_ENABLE_DEPRECATION_WARNINGS
 
 	FDelegateHandle VoipConnectionRequestDelegateHandle;
 	void OnVoipConnectionRequest(ovrMessageHandle Message, bool bIsError) const;
